[GeometryCacheLevelSequenceBaker] Add a new experimental plugin that bakes skeletal meshes in a level sequence into a geometry cache
+minor fixes for geometry cache writer and render proxy
+minor extension to leve sequence editor customization
+decoupled GPU readback processor from deformer data provider such that readback can still be processed after deformer goes away

diff --git a/Engine/Plugins/Animation/DeformerGraph/Source/OptimusCore/Private/DataInterfaces/OptimusDataInterfaceSkinnedMeshWrite.cpp b/Engine/Plugins/Animation/DeformerGraph/Source/OptimusCore/Private/DataInterfaces/OptimusDataInterfaceSkinnedMeshWrite.cpp
--- a/Engine/Plugins/Animation/DeformerGraph/Source/OptimusCore/Private/DataInterfaces/OptimusDataInterfaceSkinnedMeshWrite.cpp
+++ b/Engine/Plugins/Animation/DeformerGraph/Source/OptimusCore/Private/DataInterfaces/OptimusDataInterfaceSkinnedMeshWrite.cpp
@@ -22,6 +22,8 @@
 #include "MeshAttributes.h"
 #include "StaticMeshAttributes.h"
 #include "Async/ParallelFor.h"
+#include "OptimusGeometryReadbackProcessor.h"
+#include "Animation/MeshDeformerGeometryReadback.h"
 #endif // WITH_EDITORONLY_DATA
 
 #include UE_INLINE_GENERATED_CPP_BY_NAME(OptimusDataInterfaceSkinnedMeshWrite)
@@ -160,6 +162,11 @@ EMeshDeformerOutputBuffer UOptimusSkinnedMeshWriteDataInterface::GetOutputBuffer
 }
 
 
+UOptimusSkinnedMeshWriteDataProvider::UOptimusSkinnedMeshWriteDataProvider(FVTableHelper& Helper)
+	: Super(Helper)
+{
+}
+
 FComputeDataProviderRenderProxy* UOptimusSkinnedMeshWriteDataProvider::GetRenderProxy()
 {
 	FOptimusSkinnedMeshWriteDataProviderProxy* Proxy = new FOptimusSkinnedMeshWriteDataProviderProxy();
@@ -169,235 +176,24 @@ FComputeDataProviderRenderProxy* UOptimusSkinnedMeshWriteDataProvider::GetRender
 	Proxy->LastLodIndexPtr = &LastLodIndexCachedByRenderProxy;
 
 #if WITH_EDITORONLY_DATA
-	SetupGeometryReadbackIfNeeded(Proxy);
+	Proxy->FrameNumber = GFrameNumber;
+	Proxy->SkeletalMeshAsset = Cast<USkeletalMesh>(SkinnedMesh->GetSkinnedAsset());
+	for (TUniquePtr<FMeshDeformerGeometryReadbackRequest>& Request : GeometryReadbackRequests)
+	{
+		Proxy->GeometryReadbackRequests.Add(MoveTemp(Request));
+	}
+	
+	GeometryReadbackRequests.Empty();
 #endif
 	
 	return Proxy;
 }
-
-bool UOptimusSkinnedMeshWriteDataProvider::RequestReadbackDeformerGeometry(TFunction<void(const FMeshDescription&)> InOnReadbackCompleted_AnyThread)
-{
-#if WITH_EDITORONLY_DATA
-	OnGeometryReadbackCompletedCallbacks.Add(InOnReadbackCompleted_AnyThread);
-#endif // WITH_EDITORONLY_DATA
-
-	return true;
-}
-
-void UOptimusSkinnedMeshWriteDataProvider::Abort_RenderThread()
-{
-#if WITH_EDITORONLY_DATA
-	// When deformer is aborting, we also want to clear our pending readbacks because they will no longer be fulfilled
-	GeometryReadbacks_RenderThread.Empty();
-#endif // WITH_EDITORONLY_DATA
-}
 #if WITH_EDITORONLY_DATA
-void UOptimusSkinnedMeshWriteDataProvider::SetupGeometryReadbackIfNeeded(FOptimusSkinnedMeshWriteDataProviderProxy* Proxy)
+bool UOptimusSkinnedMeshWriteDataProvider::RequestReadbackDeformerGeometry(TUniquePtr<FMeshDeformerGeometryReadbackRequest> InRequest)
 {
-	// Process geometry readback requests if there is one
-	if (OnGeometryReadbackCompletedCallbacks.Num() > 0)
-	{
-		Proxy->RequestGeometryReadback_RenderThread = [
-			WeakThis = TWeakObjectPtr(this),
-			FrameNumber = GFrameNumber,
-			GeometryReadbackCallbacks = OnGeometryReadbackCompletedCallbacks
-		](
-			TArray<FComputeDataProviderRenderProxy::FReadbackData>& OutBufferReadbacks,
-			FSkeletalMeshObject* MeshObject,
-			FRDGBuffer* PositionBuffer,
-			FRDGBuffer* TangentBuffer,
-			FRDGBuffer* ColorBuffer
-		)
-		{
-			if (UOptimusSkinnedMeshWriteDataProvider* DataProvider = WeakThis.Get())
-			{
-				// Each proxy allocates a FGeometryReadback which stores the result of its readback and is processed async at a later time
-				TSharedRef<FGeometryReadback> GeometryReadback = MakeShared<FGeometryReadback>();
-				DataProvider->GeometryReadbacks_RenderThread.Enqueue(GeometryReadback);
-
-				GeometryReadback->FrameNumber = FrameNumber;
-				GeometryReadback->GeometryReadbackCallbacks = GeometryReadbackCallbacks;
-				GeometryReadback->LodIndex = MeshObject->GetLOD();
-
-				auto SetupBufferReadback = [WeakThis, &OutBufferReadbacks, FrameNumber](FRDGBuffer* GPUBuffer, FBufferReadback& BufferReadback)
-				{
-					if (GPUBuffer)
-					{
-						BufferReadback.bShouldReadback = true;
-						BufferReadback.OnReadbackCompleted_RenderThread = [WeakThis, &ReadbackData = BufferReadback.ReadbackData, FrameNumber](const void* InData, int InNumBytes)
-						{
-							if (ensure(InNumBytes > 0))
-							{
-								check(ReadbackData.Num() == 0);
-								ReadbackData.SetNumUninitialized(InNumBytes);
-								FPlatformMemory::Memcpy(ReadbackData.GetData(), InData, InNumBytes);
-							}
-							
-							if (UOptimusSkinnedMeshWriteDataProvider* DataProvider = WeakThis.Get())
-							{
-								// Each time a new buffer readback completes, we want to check if the geometry readback it
-								// belongs to is ready for processing. (All its buffer readbacks are ready)
-								DataProvider->ProcessCompletedGeometryReadback_RenderThread();	
-							}
-						};
-
-						// Raw buffer readback request to the compute worker
-						FComputeDataProviderRenderProxy::FReadbackData& BufferReadbackRequestRef = OutBufferReadbacks.AddDefaulted_GetRef();
-						BufferReadbackRequestRef.Buffer = GPUBuffer;
-						BufferReadbackRequestRef.NumBytes = GPUBuffer->Desc.GetSize();
-						BufferReadbackRequestRef.ReadbackCallback_RenderThread = &BufferReadback.OnReadbackCompleted_RenderThread;
-					}
-				};
+	GeometryReadbackRequests.Add(MoveTemp(InRequest));
 
-				SetupBufferReadback(PositionBuffer, GeometryReadback->Position);
-				SetupBufferReadback(TangentBuffer, GeometryReadback->Tangent);
-				SetupBufferReadback(ColorBuffer, GeometryReadback->Color);
-			}
-		};
-
-		OnGeometryReadbackCompletedCallbacks.Reset();
-	}	
-}
-#endif // WITH_EDITORONLY_DATA
-
-#if WITH_EDITORONLY_DATA
-void UOptimusSkinnedMeshWriteDataProvider::ProcessCompletedGeometryReadback_RenderThread()
-{
-	// Try to process readbacks in the order they are added,
-	// later readbacks won't be processed until earlier readbacks have been popped and processed
-	while (!GeometryReadbacks_RenderThread.IsEmpty())
-	{
-		TSharedPtr<FGeometryReadback> Readback = *GeometryReadbacks_RenderThread.Peek();
-
-		// Stop processing if a buffer is not ready
-		if ((Readback->Position.bShouldReadback && Readback->Position.ReadbackData.Num() == 0) ||
-			(Readback->Tangent.bShouldReadback && Readback->Tangent.ReadbackData.Num() == 0) ||
-			(Readback->Color.bShouldReadback && Readback->Color.ReadbackData.Num() == 0))
-		{
-			break;
-		}
-
-		// Ready for processing
-		GeometryReadbacks_RenderThread.Pop();
-		
-		auto ReadbackGeometryToMeshDescription = [WeakMesh = TWeakObjectPtr(SkinnedMesh), Readback]()
-		{
-			if (const USkeletalMeshComponent* SkeletalMeshComponent = Cast<USkeletalMeshComponent>(WeakMesh.Get()))
-			{
-				const int32 LODIndex = Readback->LodIndex;
-				const TArray<uint8>& Positions = Readback->Position.ReadbackData; 
-				const TArray<uint8>& NormalsTangents = Readback->Tangent.ReadbackData; 
-				const TArray<uint8>& Colors = Readback->Color.ReadbackData; 
-				
-				const USkeletalMesh* SkeletalMesh = SkeletalMeshComponent->GetSkeletalMeshAsset();
-				FMeshDescription MeshDescription;
-				SkeletalMesh->CloneMeshDescription(LODIndex, MeshDescription);
-
-				const TArray<int32>& RenderToImportedMap = SkeletalMesh->GetImportedModel()->LODModels[LODIndex].MeshToImportVertexMap;
-				// Avoid conditional build during parallel for loop below
-				MeshDescription.BuildVertexIndexers();
-
-				constexpr int32 SizeOfVertPosition = FSkeletalMeshDeformerHelpers::PosBufferElementMultiplier * FSkeletalMeshDeformerHelpers::PosBufferBytesPerElement; // float3, 3 * float, 12 Bytes per vert
-				constexpr int32 SizeOfVertTangents = FSkeletalMeshDeformerHelpers::TangentBufferElementMultiplier * FSkeletalMeshDeformerHelpers::TangentBufferBytesPerElement; // 2 * half4, which is 4 * 2 bytes, 16 bytes per vert
-				constexpr int32 SizeOfVertColor = FSkeletalMeshDeformerHelpers::ColorBufferBytesPerElement; // 1 * 4 * 8bit, 4 bytes per vert
-
-				const int32 NumVertPositions = Positions.Num() / SizeOfVertPosition;
-				const int32 NumVertNormalsTangents= NormalsTangents.Num() / SizeOfVertTangents;
-				const int32 NumVertColors = Colors.Num() / SizeOfVertColor;
-
-				// Make sure sizes of all buffers make sense
-				if (ensure(NumVertPositions == NumVertNormalsTangents || NumVertPositions == 0 || NumVertNormalsTangents == 0) && 
-					ensure(NumVertNormalsTangents == NumVertColors || NumVertNormalsTangents == 0 || NumVertColors == 0) &&
-					ensure(NumVertPositions == NumVertColors || NumVertPositions == 0 || NumVertColors == 0))
-				{
-					const int32 NumVertices =
-						NumVertPositions != 0 ? NumVertPositions :
-						NumVertNormalsTangents != 0 ? NumVertNormalsTangents :
-						NumVertColors != 0 ? NumVertColors : 0;
-				
-					// Parallel for each render vert
-					ParallelFor(NumVertices,[&MeshDescription, NumVertices, &RenderToImportedMap, &Positions, &NormalsTangents, &Colors](int32 RenderVertIndex)
-					{
-						const FVertexID ImportedVertexID(RenderToImportedMap[RenderVertIndex]);
-
-						// Positions
-						if (Positions.Num() > 0)
-						{
-							const float* Ptr = reinterpret_cast<const float*>(Positions.GetData());
-						
-							const float PositionX = Ptr[RenderVertIndex*3 + 0];
-							const float PositionY = Ptr[RenderVertIndex*3 + 1];
-							const float PositionZ = Ptr[RenderVertIndex*3 + 2];
-
-							MeshDescription.VertexAttributes().SetAttribute(ImportedVertexID, MeshAttribute::Vertex::Position, 0, FVector3f(PositionX,PositionY,PositionZ));
-						}
-
-						// Normals Tangents
-						const TArrayView<const FVertexInstanceID> VertexInstances = MeshDescription.GetVertexVertexInstanceIDs(ImportedVertexID);
-						if (NormalsTangents.Num() > 0)
-						{
-							const FPackedRGBA16N* Ptr = reinterpret_cast<const FPackedRGBA16N*>(NormalsTangents.GetData());
-							const FVector3f Tangent = Ptr[RenderVertIndex*2 + 0].ToFVector3f();
-							const FVector4f Normal = Ptr[RenderVertIndex*2 + 1].ToFVector4f();
-						
-							for (const FVertexInstanceID VertexInstanceID : VertexInstances)
-							{
-								MeshDescription.VertexInstanceAttributes().SetAttribute<FVector3f>(
-									VertexInstanceID,
-									MeshAttribute::VertexInstance::Normal,
-									0,
-									FVector3f(Normal));
-							
-								MeshDescription.VertexInstanceAttributes().SetAttribute<FVector3f>(
-									VertexInstanceID,
-									MeshAttribute::VertexInstance::Tangent,
-									0,
-									Tangent);
-
-								MeshDescription.VertexInstanceAttributes().SetAttribute<float>(
-									VertexInstanceID,
-									MeshAttribute::VertexInstance::BinormalSign,
-									0,
-									Normal.W);
-							}
-						}
-
-						// Colors
-						if (Colors.Num() > 0)
-						{
-							check(sizeof(FColor) == 4);
-						
-							const FColor* Ptr = reinterpret_cast<const FColor*>(Colors.GetData());
-
-							for (const FVertexInstanceID VertexInstanceID : VertexInstances)
-							{
-								MeshDescription.VertexInstanceAttributes().SetAttribute<FVector4f>(
-									VertexInstanceID,
-									MeshAttribute::VertexInstance::Color,
-									0,
-									FVector4f(Ptr[RenderVertIndex]));
-							}
-						}
-					});
-
-					for (TFunction<void(const FMeshDescription&)> Callback : Readback->GeometryReadbackCallbacks)
-					{
-						Callback(MeshDescription);	
-					}
-				}
-			}	
-		};
-
-		// Readback data should be processed sequentially on a background thread, so using last task as prerequisite
-		if (LastReadbackProcessingTask_RenderThread)
-		{
-			LastReadbackProcessingTask_RenderThread = FFunctionGraphTask::CreateAndDispatchWhenReady(ReadbackGeometryToMeshDescription, TStatId(), LastReadbackProcessingTask_RenderThread, ENamedThreads::AnyBackgroundThreadNormalTask);	
-		}
-		else
-		{
-			LastReadbackProcessingTask_RenderThread = FFunctionGraphTask::CreateAndDispatchWhenReady(ReadbackGeometryToMeshDescription, TStatId(), NULL, ENamedThreads::AnyBackgroundThreadNormalTask);	
-		}
-	}
+	return true;
 }
 #endif // WITH_EDITORONLY_DATA
 
@@ -445,7 +241,7 @@ void FOptimusSkinnedMeshWriteDataProviderProxy::AllocateResources(FRDGBuilder& G
 	else
 	{
 #if WITH_EDITORONLY_DATA
-		if (RequestGeometryReadback_RenderThread)
+		if (GeometryReadbackRequests.Num() > 0)
 		{
 			PositionBuffer = FSkeletalMeshDeformerHelpers::GetAllocatedPositionBuffer(GraphBuilder, SkeletalMeshObject, LodIndex);
 		}
@@ -467,7 +263,7 @@ void FOptimusSkinnedMeshWriteDataProviderProxy::AllocateResources(FRDGBuilder& G
 	else
 	{
 #if WITH_EDITORONLY_DATA
-		if (RequestGeometryReadback_RenderThread)
+		if (GeometryReadbackRequests.Num() > 0)
 		{
 			TangentBuffer = FSkeletalMeshDeformerHelpers::GetAllocatedTangentBuffer(GraphBuilder, SkeletalMeshObject, LodIndex);
 		}
@@ -486,7 +282,7 @@ void FOptimusSkinnedMeshWriteDataProviderProxy::AllocateResources(FRDGBuilder& G
 	else
 	{
 #if WITH_EDITORONLY_DATA
-		if (RequestGeometryReadback_RenderThread)
+		if (GeometryReadbackRequests.Num() > 0)
 		{
 			ColorBuffer = FSkeletalMeshDeformerHelpers::GetAllocatedColorBuffer(GraphBuilder, SkeletalMeshObject, LodIndex);	
 		}
@@ -522,9 +318,53 @@ void FOptimusSkinnedMeshWriteDataProviderProxy::GatherDispatchData(FDispatchData
 void FOptimusSkinnedMeshWriteDataProviderProxy::GetReadbackData(TArray<FReadbackData>& OutReadbackData) const
 {
 #if WITH_EDITORONLY_DATA
-	if (RequestGeometryReadback_RenderThread)
+	if (GeometryReadbackRequests.Num() > 0)
 	{
-		RequestGeometryReadback_RenderThread(OutReadbackData, SkeletalMeshObject, PositionBuffer, TangentBuffer, ColorBuffer);
+		// Each proxy allocates a FGeometryReadback which stores the result of its readback and is processed async at a later time
+		using FGeometryReadback = FOptimusGeometryReadbackProcessor::FGeometryReadback;
+		using FBufferReadback = FOptimusGeometryReadbackProcessor::FBufferReadback;
+		using FReadbackProcessor = FOptimusGeometryReadbackProcessor;
+		
+		TSharedPtr<FGeometryReadback> GeometryReadback = MakeShared<FGeometryReadback>();
+
+		GeometryReadback->FrameNumber = FrameNumber;
+		for (TUniquePtr<FMeshDeformerGeometryReadbackRequest>& Request : GeometryReadbackRequests)
+		{
+			GeometryReadback->GeometryReadbackRequests.Add(MoveTemp(Request));
+		}
+		GeometryReadback->SkeletalMesh = SkeletalMeshAsset;
+		GeometryReadback->LodIndex = SkeletalMeshObject->GetLOD();
+
+		auto SetupBufferReadback = [&OutReadbackData, this](FRDGBuffer* GPUBuffer, FBufferReadback& BufferReadback)
+		{
+			if (GPUBuffer)
+			{
+				BufferReadback.bShouldReadback = true;
+				BufferReadback.OnReadbackCompleted_RenderThread = [&ReadbackData = BufferReadback.ReadbackData](const void* InData, int InNumBytes)
+				{
+					if (ensure(InNumBytes > 0))
+					{
+						check(ReadbackData.Num() == 0);
+						ReadbackData.SetNumUninitialized(InNumBytes);
+						FPlatformMemory::Memcpy(ReadbackData.GetData(), InData, InNumBytes);
+					}
+					
+					FReadbackProcessor::Get().ProcessCompletedGeometryReadback_RenderThread();
+				};
+
+				// Raw buffer readback request to the compute worker
+				FComputeDataProviderRenderProxy::FReadbackData& BufferReadbackRequestRef = OutReadbackData.AddDefaulted_GetRef();
+				BufferReadbackRequestRef.Buffer = GPUBuffer;
+				BufferReadbackRequestRef.NumBytes = GPUBuffer->Desc.GetSize();
+				BufferReadbackRequestRef.ReadbackCallback_RenderThread = &BufferReadback.OnReadbackCompleted_RenderThread;
+			}
+		};
+
+		SetupBufferReadback(PositionBuffer, GeometryReadback->Position);
+		SetupBufferReadback(TangentBuffer, GeometryReadback->Tangent);
+		SetupBufferReadback(ColorBuffer, GeometryReadback->Color);
+		
+		FReadbackProcessor::Get().Add(GeometryReadback);	
 	}
 #endif // WITH_EDITORONLY_DATA
 }
diff --git a/Engine/Plugins/Animation/DeformerGraph/Source/OptimusCore/Private/DataInterfaces/OptimusDataInterfaceSkinnedMeshWrite.h b/Engine/Plugins/Animation/DeformerGraph/Source/OptimusCore/Private/DataInterfaces/OptimusDataInterfaceSkinnedMeshWrite.h
--- a/Engine/Plugins/Animation/DeformerGraph/Source/OptimusCore/Private/DataInterfaces/OptimusDataInterfaceSkinnedMeshWrite.h
+++ b/Engine/Plugins/Animation/DeformerGraph/Source/OptimusCore/Private/DataInterfaces/OptimusDataInterfaceSkinnedMeshWrite.h
@@ -7,9 +7,6 @@
 #include "OptimusComputeDataInterface.h"
 #include "ComputeFramework/ComputeDataProvider.h"
 
-#if WITH_EDITORONLY_DATA
-#include "Containers/Queue.h"
-#endif // WITH_EDITORONLY_DATA
 
 #include "OptimusDataInterfaceSkinnedMeshWrite.generated.h"
 
@@ -19,6 +16,10 @@ class FSkeletalMeshObject;
 class FSkinedMeshWriteDataInterfaceParameters;
 class USkinnedMeshComponent;
 
+#if WITH_EDITORONLY_DATA
+class USkeletalMesh;
+#endif // WITH_EDITORONLY_DATA
+
 /** Compute Framework Data Interface for writing skinned mesh. */
 UCLASS(Category = ComputeFramework)
 class OPTIMUSCORE_API UOptimusSkinnedMeshWriteDataInterface :
@@ -66,6 +67,10 @@ class UOptimusSkinnedMeshWriteDataProvider :
 	GENERATED_BODY()
 
 public:
+	UOptimusSkinnedMeshWriteDataProvider() = default;
+	~UOptimusSkinnedMeshWriteDataProvider() = default;
+	UOptimusSkinnedMeshWriteDataProvider(FVTableHelper& Helper);
+	
 	UPROPERTY(EditAnywhere, BlueprintReadWrite, Category = Binding)
 	TObjectPtr<USkinnedMeshComponent> SkinnedMesh = nullptr;
 
@@ -75,47 +80,18 @@ public:
 	FComputeDataProviderRenderProxy* GetRenderProxy() override;
 	//~ End UComputeDataProvider Interface
 	
-	
+#if WITH_EDITORONLY_DATA	
 	//~ Begin IOptimusDeformerGeometryReadbackProvider Interface
-	bool RequestReadbackDeformerGeometry(TFunction<void(const FMeshDescription&)> InOnReadbackCompleted_AnyThread) override;
-	void Abort_RenderThread() override;
+	bool RequestReadbackDeformerGeometry(TUniquePtr<FMeshDeformerGeometryReadbackRequest> InRequest) override;
 	//~ End IOptimusDeformerGeometryReadbackProvider Interface
-	
-	
+#endif // WITH_EDITORONLY_DATA	
 private:
 	// Served as persistent storage for the provider proxy, should not be used by the data provider itself
 	int32 LastLodIndexCachedByRenderProxy = 0;
 
 #if WITH_EDITORONLY_DATA
 	/** Readback requests for the current frame */
-	TArray<TFunction<void(const FMeshDescription&)>> OnGeometryReadbackCompletedCallbacks;
-
-	void SetupGeometryReadbackIfNeeded(class FOptimusSkinnedMeshWriteDataProviderProxy* Proxy);
-	void ProcessCompletedGeometryReadback_RenderThread();
-	
-	struct FBufferReadback
-	{
-		bool bShouldReadback = false;
-		TArray<uint8> ReadbackData;
-		FComputeDataProviderRenderProxy::FReadbackCallback OnReadbackCompleted_RenderThread;
-	};
-	
-	struct FGeometryReadback
-	{
-		uint64 FrameNumber = 0;
-		TArray<TFunction<void(const FMeshDescription&)>> GeometryReadbackCallbacks;
-		
-		int32 LodIndex = 0;
-		
-		FBufferReadback Position;
-		FBufferReadback Tangent;
-		FBufferReadback Color;
-	};
-	
-	// Served as persistent storage for the proxy, render thread only
-	TQueue<TSharedPtr<FGeometryReadback>> GeometryReadbacks_RenderThread;
-
-	FGraphEventRef LastReadbackProcessingTask_RenderThread;
+	TArray<TUniquePtr<FMeshDeformerGeometryReadbackRequest>> GeometryReadbackRequests;
 #endif // WITH_EDITORONLY_DATA
 };
 
@@ -134,6 +110,13 @@ public:
 	int32* LastLodIndexPtr = nullptr;
 
 #if WITH_EDITORONLY_DATA
+
+	int32 FrameNumber = 0;
+
+	USkeletalMesh* SkeletalMeshAsset = nullptr;
+	/** Readback requests for the current frame */
+	mutable TArray<TUniquePtr<FMeshDeformerGeometryReadbackRequest>> GeometryReadbackRequests;
+	
 	TFunction<
 		void(
 			TArray<FReadbackData>&,
diff --git a/Engine/Plugins/Animation/DeformerGraph/Source/OptimusCore/Private/OptimusDeformerDynamicInstanceManager.cpp b/Engine/Plugins/Animation/DeformerGraph/Source/OptimusCore/Private/OptimusDeformerDynamicInstanceManager.cpp
--- a/Engine/Plugins/Animation/DeformerGraph/Source/OptimusCore/Private/OptimusDeformerDynamicInstanceManager.cpp
+++ b/Engine/Plugins/Animation/DeformerGraph/Source/OptimusCore/Private/OptimusDeformerDynamicInstanceManager.cpp
@@ -8,6 +8,10 @@
 #include "OptimusDeformerInstance.h"
 #include "ControlRig/RigUnit_Optimus.h"
 
+#if WITH_EDITORONLY_DATA
+#include "Animation/MeshDeformerGeometryReadback.h"
+#endif
+
 void UOptimusDeformerDynamicInstanceManager::AllocateResources()
 {
 	// Typically called during recreate render state
@@ -118,9 +122,9 @@ void UOptimusDeformerDynamicInstanceManager::EnqueueWork(FEnqueueWorkDesc const&
 		// Readback the deformer geometry after the last deformer executes
 		if (Index == SortedInstances.Num() - 1)
 		{
-			for (const TFunction<void(const FMeshDescription&)>& Callback : OnGeometryReadbackCompletedCallbacks)
+			for (TUniquePtr<FMeshDeformerGeometryReadbackRequest>& Request : GeometryReadbackRequests)
 			{
-				Instance->RequestReadbackDeformerGeometry(Callback);
+				Instance->RequestReadbackDeformerGeometry(MoveTemp(Request));
 			}
 		}
 #endif // WITH_EDITORONLY_DATA
@@ -132,7 +136,7 @@ void UOptimusDeformerDynamicInstanceManager::EnqueueWork(FEnqueueWorkDesc const&
 	ExecutionQueueMap.Reset();
 #if WITH_EDITORONLY_DATA
 	// Avoid unbounded accumulation of readback requests
-	OnGeometryReadbackCompletedCallbacks.Reset();
+	GeometryReadbackRequests.Reset();
 #endif // WITH_EDITORONLY_DATA
 }
 
@@ -143,11 +147,12 @@ EMeshDeformerOutputBuffer UOptimusDeformerDynamicInstanceManager::GetOutputBuffe
 	
 	return Result;
 }
+
 #if WITH_EDITORONLY_DATA
-bool UOptimusDeformerDynamicInstanceManager::RequestReadbackDeformerGeometry(TFunction<void(const FMeshDescription&)> InOnReadbackCompleted_AnyThread)
+bool UOptimusDeformerDynamicInstanceManager::RequestReadbackDeformerGeometry(TUniquePtr<FMeshDeformerGeometryReadbackRequest> InRequest)
 {
 	// These requests are forwarded to the last deformer instance that runs, see EnqueueWork
-	OnGeometryReadbackCompletedCallbacks.Add(InOnReadbackCompleted_AnyThread);
+	GeometryReadbackRequests.Add(MoveTemp(InRequest));
 	
 	return true;
 }
diff --git a/Engine/Plugins/Animation/DeformerGraph/Source/OptimusCore/Private/OptimusDeformerInstance.cpp b/Engine/Plugins/Animation/DeformerGraph/Source/OptimusCore/Private/OptimusDeformerInstance.cpp
--- a/Engine/Plugins/Animation/DeformerGraph/Source/OptimusCore/Private/OptimusDeformerInstance.cpp
+++ b/Engine/Plugins/Animation/DeformerGraph/Source/OptimusCore/Private/OptimusDeformerInstance.cpp
@@ -498,9 +498,9 @@ void UOptimusDeformerInstance::AllocateResources()
 
 void UOptimusDeformerInstance::ReleaseResources()
 {
-	if (Scene || BufferPool || WeakGeometryReadbackProvider.IsValid())
+	if (Scene || BufferPool)
 	{
-		ENQUEUE_RENDER_COMMAND(OptimusReleaseResources)([BufferPool=MoveTemp(BufferPool), Scene = Scene, OwnerPointer = this, WeakGeometryReadbackProvider = WeakGeometryReadbackProvider] (FRHICommandListImmediate& InCmdList)
+		ENQUEUE_RENDER_COMMAND(OptimusReleaseResources)([BufferPool=MoveTemp(BufferPool), Scene = Scene, OwnerPointer = this] (FRHICommandListImmediate& InCmdList)
 		{
 			if (Scene)
 			{
@@ -511,11 +511,6 @@ void UOptimusDeformerInstance::ReleaseResources()
 			{
 				BufferPool->ReleaseResources();
 			}
-
-			if (WeakGeometryReadbackProvider.IsValid())
-			{
-				WeakGeometryReadbackProvider->Abort_RenderThread();
-			}
 		});
 	}
 }
@@ -608,11 +603,11 @@ EMeshDeformerOutputBuffer UOptimusDeformerInstance::GetOutputBuffers() const
 }
 
 #if WITH_EDITORONLY_DATA
-bool UOptimusDeformerInstance::RequestReadbackDeformerGeometry(TFunction<void(const FMeshDescription&)> InOnReadbackCompleted_AnyThread)
+bool UOptimusDeformerInstance::RequestReadbackDeformerGeometry(TUniquePtr<FMeshDeformerGeometryReadbackRequest> InRequest)
 {
 	if (WeakGeometryReadbackProvider.IsValid())
 	{
-		WeakGeometryReadbackProvider->RequestReadbackDeformerGeometry(InOnReadbackCompleted_AnyThread);
+		WeakGeometryReadbackProvider->RequestReadbackDeformerGeometry(MoveTemp(InRequest));
 		return true;
 	}
 
diff --git a/Engine/Plugins/Animation/DeformerGraph/Source/OptimusCore/Private/OptimusGeometryReadbackProcessor.cpp b/Engine/Plugins/Animation/DeformerGraph/Source/OptimusCore/Private/OptimusGeometryReadbackProcessor.cpp
new file mode 100644
--- /dev/null
+++ b/Engine/Plugins/Animation/DeformerGraph/Source/OptimusCore/Private/OptimusGeometryReadbackProcessor.cpp
@@ -0,0 +1,208 @@
+// Copyright Epic Games, Inc. All Rights Reserved.
+
+#include "OptimusGeometryReadbackProcessor.h"
+
+#include "MeshAttributes.h"
+#include "SkeletalMeshDeformerHelpers.h"
+#include "StaticMeshAttributes.h"
+#include "Animation/MeshDeformerGeometryReadback.h"
+#include "Rendering/SkeletalMeshModel.h"
+#include "Async/ParallelFor.h"
+
+#if WITH_EDITORONLY_DATA
+
+// At the moment there are meshes with MeshToImportVertexMap not matching their mesh description for some reason
+// Once the root cause is determined we can remove this
+#define USE_MESH_TO_IMPORT_VERTEX_MAP 0
+
+FOptimusGeometryReadbackProcessor& FOptimusGeometryReadbackProcessor::Get()
+{
+	static FOptimusGeometryReadbackProcessor Instance;
+	return Instance;
+}
+
+void FOptimusGeometryReadbackProcessor::Add(TSharedPtr<FGeometryReadback> InReadback)
+{
+	Get().GeometryReadbacks.Enqueue(InReadback);
+}
+
+void FOptimusGeometryReadbackProcessor::ProcessCompletedGeometryReadback_RenderThread()
+{
+	// Try to process readbacks in the order they are added,
+	// later readbacks won't be processed until earlier readbacks have been popped and processed
+	while (!GeometryReadbacks.IsEmpty())
+	{
+		TSharedPtr<FGeometryReadback> Readback = *GeometryReadbacks.Peek();
+
+		// Stop processing if a buffer is not ready
+		if ((Readback->Position.bShouldReadback && Readback->Position.ReadbackData.Num() == 0) ||
+			(Readback->Tangent.bShouldReadback && Readback->Tangent.ReadbackData.Num() == 0) ||
+			(Readback->Color.bShouldReadback && Readback->Color.ReadbackData.Num() == 0))
+		{
+			break;
+		}
+
+		// Ready for processing
+		GeometryReadbacks.Dequeue();
+		
+		auto ReadbackGeometryToMeshDescription = [Readback]()
+		{
+			{
+				const int32 LODIndex = Readback->LodIndex;
+				const TArray<uint8>& Positions = Readback->Position.ReadbackData; 
+				const TArray<uint8>& NormalsTangents = Readback->Tangent.ReadbackData; 
+				const TArray<uint8>& Colors = Readback->Color.ReadbackData; 
+
+				if (const USkeletalMesh* SkeletalMesh = Readback->SkeletalMesh.Get())
+				{
+					FMeshDescription MeshDescription;
+					SkeletalMesh->CloneMeshDescription(LODIndex, MeshDescription);
+					
+					// Avoid conditional build during parallel for loop below
+					MeshDescription.BuildVertexIndexers();
+
+					constexpr int32 SizeOfVertPosition = FSkeletalMeshDeformerHelpers::PosBufferElementMultiplier * FSkeletalMeshDeformerHelpers::PosBufferBytesPerElement; // float3, 3 * float, 12 Bytes per vert
+					constexpr int32 SizeOfVertTangents = FSkeletalMeshDeformerHelpers::TangentBufferElementMultiplier * FSkeletalMeshDeformerHelpers::TangentBufferBytesPerElement; // 2 * half4, which is 4 * 2 bytes, 16 bytes per vert
+					constexpr int32 SizeOfVertColor = FSkeletalMeshDeformerHelpers::ColorBufferBytesPerElement; // 1 * 4 * 8bit, 4 bytes per vert
+
+					const int32 NumVertPositions = Positions.Num() / SizeOfVertPosition;
+					const int32 NumVertNormalsTangents= NormalsTangents.Num() / SizeOfVertTangents;
+					const int32 NumVertColors = Colors.Num() / SizeOfVertColor;
+
+					// Make sure sizes of all buffers make sense
+					if (ensure(NumVertPositions == NumVertNormalsTangents || NumVertPositions == 0 || NumVertNormalsTangents == 0) && 
+						ensure(NumVertNormalsTangents == NumVertColors || NumVertNormalsTangents == 0 || NumVertColors == 0) &&
+						ensure(NumVertPositions == NumVertColors || NumVertPositions == 0 || NumVertColors == 0))
+					{
+						const int32 NumVertices =
+							NumVertPositions != 0 ? NumVertPositions :
+							NumVertNormalsTangents != 0 ? NumVertNormalsTangents :
+							NumVertColors != 0 ? NumVertColors : 0;
+						TMeshAttributesRef<FVertexID, FVector3f> PositionAttribute =
+							MeshDescription.VertexAttributes().GetAttributesRef<FVector3f>(MeshAttribute::Vertex::Position);
+						TMeshAttributesRef<FVertexInstanceID, FVector3f> NormalAttribute =
+							MeshDescription.VertexInstanceAttributes().GetAttributesRef<FVector3f>(MeshAttribute::VertexInstance::Normal);
+						TMeshAttributesRef<FVertexInstanceID, FVector3f> TangentAttribute = 
+							MeshDescription.VertexInstanceAttributes().GetAttributesRef<FVector3f>(MeshAttribute::VertexInstance::Tangent);
+						TMeshAttributesRef<FVertexInstanceID, float> BinormalSignAttribute = 
+							MeshDescription.VertexInstanceAttributes().GetAttributesRef<float>(MeshAttribute::VertexInstance::BinormalSign);
+						TMeshAttributesRef<FVertexID, FVector4f> ColorAttribute =
+							MeshDescription.VertexAttributes().GetAttributesRef<FVector4f>(MeshAttribute::VertexInstance::Color);
+						
+						FSkeletalMeshLODModel& LODModel = SkeletalMesh->GetImportedModel()->LODModels[LODIndex];
+
+	#if USE_MESH_TO_IMPORT_VERTEX_MAP
+						const TArray<int32>& RenderToImportedMap = LODModel.MeshToImportVertexMap;
+						uint32 MaxImportVertex = LODModel.MaxImportVertex;
+	#else
+						const TArray<uint32>& RenderToImportedMap = LODModel.GetRawPointIndices();
+						uint32 MaxImportVertex = 0;
+						for (uint32 Raw : RenderToImportedMap)
+						{
+							MaxImportVertex = FMath::Max(MaxImportVertex, Raw);
+						}
+	#endif
+						
+						TArray<std::atomic<int32>> ImportedVertWriteCount;
+						ImportedVertWriteCount.SetNum(MaxImportVertex+1);
+						for(std::atomic<int32>& Counter : ImportedVertWriteCount)
+						{
+							Counter = 0;
+						}
+						
+						// Parallel for each render vert
+						ParallelFor(NumVertices,
+							[
+							&Readback,
+							&MeshDescription,
+							&PositionAttribute, 
+							&NormalAttribute, 
+							&TangentAttribute, 
+							&BinormalSignAttribute,
+							&ColorAttribute, 
+							NumVertices,
+							&RenderToImportedMap, 
+							&Positions, 
+							&NormalsTangents,
+							&Colors,
+							&ImportedVertWriteCount
+							](int32 RenderVertIndex)
+						{
+							// int32 ImportedVertexIndex = RenderToImportedMap[RenderVertIndex];
+							uint32 ImportedVertexIndex = RenderToImportedMap[RenderVertIndex];
+							const FVertexID ImportedVertexID(ImportedVertexIndex);
+
+							// Only write to each imported vert once
+							if (++ImportedVertWriteCount[ImportedVertexIndex] > 1)
+							{
+								return;
+							}
+							
+							// Positions
+							if (Positions.Num() > 0)
+							{
+								const float* Ptr = reinterpret_cast<const float*>(Positions.GetData());
+							
+								const float PositionX = Ptr[RenderVertIndex*3 + 0];
+								const float PositionY = Ptr[RenderVertIndex*3 + 1];
+								const float PositionZ = Ptr[RenderVertIndex*3 + 2];
+
+								PositionAttribute[ImportedVertexID] = FVector3f(PositionX,PositionY,PositionZ);
+							}
+
+							// Normals Tangents
+							const TArrayView<const FVertexInstanceID> VertexInstances = MeshDescription.GetVertexVertexInstanceIDs(ImportedVertexID);
+							if (NormalsTangents.Num() > 0)
+							{
+								const FPackedRGBA16N* Ptr = reinterpret_cast<const FPackedRGBA16N*>(NormalsTangents.GetData());
+								const FVector3f Tangent = Ptr[RenderVertIndex*2 + 0].ToFVector3f();
+								const FVector4f Normal = Ptr[RenderVertIndex*2 + 1].ToFVector4f();
+							
+								for (const FVertexInstanceID VertexInstanceID : VertexInstances)
+								{
+									NormalAttribute[VertexInstanceID] = FVector3f(Normal);
+			
+									TangentAttribute[VertexInstanceID] = Tangent;	
+
+									BinormalSignAttribute[VertexInstanceID] = Normal.W;
+								}
+							}
+
+							// Colors
+							if (Colors.Num() > 0 && ColorAttribute.IsValid())
+							{
+								check(sizeof(FColor) == 4);
+							
+								const FColor* Ptr = reinterpret_cast<const FColor*>(Colors.GetData());
+
+								for (const FVertexInstanceID VertexInstanceID : VertexInstances)
+								{
+									ColorAttribute[VertexInstanceID] = FVector4f(Ptr[RenderVertIndex]);
+								}
+							}
+						});
+
+						for (TUniquePtr<FMeshDeformerGeometryReadbackRequest>& Request : Readback->GeometryReadbackRequests)
+						{
+							Request->MeshDescriptionCallback_AnyThread(MeshDescription);
+							Request->bHandled = true;
+						}
+					}
+				}
+				
+			}	
+		};
+
+		// Readback data should be processed sequentially on a background thread, so using last task as prerequisite
+		if (LastReadbackProcessingTask_RenderThread)
+		{
+			LastReadbackProcessingTask_RenderThread = FFunctionGraphTask::CreateAndDispatchWhenReady(ReadbackGeometryToMeshDescription, TStatId(), LastReadbackProcessingTask_RenderThread, ENamedThreads::AnyBackgroundThreadNormalTask);	
+		}
+		else
+		{
+			LastReadbackProcessingTask_RenderThread = FFunctionGraphTask::CreateAndDispatchWhenReady(ReadbackGeometryToMeshDescription, TStatId(), NULL, ENamedThreads::AnyBackgroundThreadNormalTask);	
+		}
+	}	
+}
+
+#endif // WITH_EDITORONLY_DATA
diff --git a/Engine/Plugins/Animation/DeformerGraph/Source/OptimusCore/Private/OptimusGeometryReadbackProcessor.h b/Engine/Plugins/Animation/DeformerGraph/Source/OptimusCore/Private/OptimusGeometryReadbackProcessor.h
new file mode 100644
--- /dev/null
+++ b/Engine/Plugins/Animation/DeformerGraph/Source/OptimusCore/Private/OptimusGeometryReadbackProcessor.h
@@ -0,0 +1,44 @@
+// Copyright Epic Games, Inc. All Rights Reserved.
+
+#pragma once
+#include "ComputeFramework/ComputeDataProvider.h"
+#include "Engine/SkeletalMesh.h"
+#include "Containers/MpscQueue.h"
+#include "Animation/MeshDeformerGeometryReadback.h"
+
+#if WITH_EDITORONLY_DATA
+class FOptimusGeometryReadbackProcessor
+{
+public:
+
+	struct FBufferReadback
+	{
+		bool bShouldReadback = false;
+		TArray<uint8> ReadbackData;
+		FComputeDataProviderRenderProxy::FReadbackCallback OnReadbackCompleted_RenderThread;
+	};
+	
+	struct FGeometryReadback
+	{
+		uint64 FrameNumber = 0;
+		TArray<TUniquePtr<FMeshDeformerGeometryReadbackRequest>> GeometryReadbackRequests;
+
+		TWeakObjectPtr<USkeletalMesh> SkeletalMesh;
+		int32 LodIndex = 0;
+		
+		FBufferReadback Position;
+		FBufferReadback Tangent;
+		FBufferReadback Color;
+	};
+
+	static FOptimusGeometryReadbackProcessor& Get();
+	void Add(TSharedPtr<FGeometryReadback>);
+	void ProcessCompletedGeometryReadback_RenderThread();
+
+private:
+	
+	// All data provider shares the same processor singleton
+	TMpscQueue<TSharedPtr<FGeometryReadback>> GeometryReadbacks;
+	FGraphEventRef LastReadbackProcessingTask_RenderThread;	
+};
+#endif
diff --git a/Engine/Plugins/Animation/DeformerGraph/Source/OptimusCore/Public/IOptimusDeformerGeometryReadbackProvider.h b/Engine/Plugins/Animation/DeformerGraph/Source/OptimusCore/Public/IOptimusDeformerGeometryReadbackProvider.h
--- a/Engine/Plugins/Animation/DeformerGraph/Source/OptimusCore/Public/IOptimusDeformerGeometryReadbackProvider.h
+++ b/Engine/Plugins/Animation/DeformerGraph/Source/OptimusCore/Public/IOptimusDeformerGeometryReadbackProvider.h
@@ -3,13 +3,12 @@
 #pragma once
 
 #include "UObject/Interface.h"
-
+#if WITH_EDITORONLY_DATA
+#include "Animation/MeshDeformerGeometryReadback.h"
+#endif // WWITH_EDITORONLY_DATA
 #include "IOptimusDeformerGeometryReadbackProvider.generated.h"
 
 
-class UOptimusDeformerInstance;
-struct FMeshDescription;
-
 UINTERFACE()
 class OPTIMUSCORE_API UOptimusDeformerGeometryReadbackProvider :
 	public UInterface
@@ -23,6 +22,7 @@ class IOptimusDeformerGeometryReadbackProvider
 	GENERATED_BODY()
 
 public:
-	virtual bool RequestReadbackDeformerGeometry(TFunction<void(const FMeshDescription&)> InOnReadbackCompleted_AnyThread) = 0;
-	virtual void Abort_RenderThread() = 0;
+#if WITH_EDITORONLY_DATA
+	virtual bool RequestReadbackDeformerGeometry(TUniquePtr<FMeshDeformerGeometryReadbackRequest> InRequest) = 0;
+#endif // WWITH_EDITORONLY_DATA
 };
diff --git a/Engine/Plugins/Animation/DeformerGraph/Source/OptimusCore/Public/OptimusDeformerDynamicInstanceManager.h b/Engine/Plugins/Animation/DeformerGraph/Source/OptimusCore/Public/OptimusDeformerDynamicInstanceManager.h
--- a/Engine/Plugins/Animation/DeformerGraph/Source/OptimusCore/Public/OptimusDeformerDynamicInstanceManager.h
+++ b/Engine/Plugins/Animation/DeformerGraph/Source/OptimusCore/Public/OptimusDeformerDynamicInstanceManager.h
@@ -4,6 +4,11 @@
 
 #include "Animation/MeshDeformerInstance.h"
 #include "Animation/MeshDeformerProducer.h"
+
+#if WITH_EDITORONLY_DATA
+#include "Animation/MeshDeformerGeometryReadback.h"
+#endif // WTIH_EDITORONLY_DATA
+
 #include "OptimusDeformerDynamicInstanceManager.generated.h"
 
 class UControlRig;
@@ -41,7 +46,7 @@ public:
 
 #if WITH_EDITORONLY_DATA
 	/** Readback the deformed geometry as mesh description */
-	bool RequestReadbackDeformerGeometry(TFunction<void(const FMeshDescription&)> InOnReadbackCompleted_AnyThread) override;
+	bool RequestReadbackDeformerGeometry(TUniquePtr<FMeshDeformerGeometryReadbackRequest> InRequest) override;
 #endif // WITH_EDITORONLY_DATA
 	
 	/** InstanceManager is an intermediate instance, call this function to get the instance for the deformer that created this instance manager */
@@ -89,7 +94,7 @@ private:
 
 #if WITH_EDITORONLY_DATA
 	/** Readback requests for the current frame */
-	TArray<TFunction<void(const FMeshDescription&)>> OnGeometryReadbackCompletedCallbacks;
+	TArray<TUniquePtr<FMeshDeformerGeometryReadbackRequest>> GeometryReadbackRequests;
 #endif // WITH_EDITORONLY_DATA
 };
 
diff --git a/Engine/Plugins/Animation/DeformerGraph/Source/OptimusCore/Public/OptimusDeformerInstance.h b/Engine/Plugins/Animation/DeformerGraph/Source/OptimusCore/Public/OptimusDeformerInstance.h
--- a/Engine/Plugins/Animation/DeformerGraph/Source/OptimusCore/Public/OptimusDeformerInstance.h
+++ b/Engine/Plugins/Animation/DeformerGraph/Source/OptimusCore/Public/OptimusDeformerInstance.h
@@ -319,7 +319,7 @@ protected:
 	void EnqueueWork(FEnqueueWorkDesc const& InDesc) override;
 	EMeshDeformerOutputBuffer GetOutputBuffers() const override;
 #if WITH_EDITORONLY_DATA
-	bool RequestReadbackDeformerGeometry(TFunction<void(const FMeshDescription&)> InOnReadbackCompleted_AnyThread) override;
+	bool RequestReadbackDeformerGeometry(TUniquePtr<FMeshDeformerGeometryReadbackRequest> InRequest) override;
 #endif // WITH_EDITORONLY_DATA
 	
 	UMeshDeformerInstance* GetInstanceForSourceDeformer() override {return this;};
diff --git a/Engine/Plugins/Experimental/GeometryCacheLevelSequenceBaker/GeometryCacheLevelSequenceBaker.uplugin b/Engine/Plugins/Experimental/GeometryCacheLevelSequenceBaker/GeometryCacheLevelSequenceBaker.uplugin
new file mode 100644
--- /dev/null
+++ b/Engine/Plugins/Experimental/GeometryCacheLevelSequenceBaker/GeometryCacheLevelSequenceBaker.uplugin
@@ -0,0 +1,37 @@
+﻿{
+	"FileVersion": 3,
+	"Version": 1,
+	"VersionName": "1.0",
+	"FriendlyName": "GeometryCacheLevelSequenceBaker",
+	"Description": "Bake Skeletal Meshes in Level Sequence to Geometry Cache",
+	"Category": "Other",
+	"CreatedBy" : "Epic Games, Inc.",
+	"CreatedByURL" : "https://epicgames.com",
+	"DocsURL": "",
+	"MarketplaceURL": "",
+	"CanContainContent": true,
+	"IsBetaVersion": false,
+	"IsExperimentalVersion": true,
+	"Installed": false,
+	"Modules": [
+		{
+			"Name": "GeometryCacheLevelSequenceBaker",
+			"Type": "Editor",
+			"LoadingPhase": "PostEngineInit"
+		}
+	],
+	"Plugins": [
+		{
+			"Name" : "LevelSequenceEditor",
+			"Enabled" : true
+		},
+		{
+			"Name" : "GeometryCache",
+			"Enabled" : true
+		},
+		{
+			"Name" : "DeformerGraph",
+			"Enabled" : true
+		}
+	]
+}
\ No newline at end of file
diff --git a/Engine/Plugins/Experimental/GeometryCacheLevelSequenceBaker/Source/GeometryCacheLevelSequenceBaker/GeometryCacheLevelSequenceBaker.Build.cs b/Engine/Plugins/Experimental/GeometryCacheLevelSequenceBaker/Source/GeometryCacheLevelSequenceBaker/GeometryCacheLevelSequenceBaker.Build.cs
new file mode 100644
--- /dev/null
+++ b/Engine/Plugins/Experimental/GeometryCacheLevelSequenceBaker/Source/GeometryCacheLevelSequenceBaker/GeometryCacheLevelSequenceBaker.Build.cs
@@ -0,0 +1,63 @@
+﻿// Copyright Epic Games, Inc. All Rights Reserved.
+
+using UnrealBuildTool;
+
+public class GeometryCacheLevelSequenceBaker : ModuleRules
+{
+	public GeometryCacheLevelSequenceBaker(ReadOnlyTargetRules Target) : base(Target)
+	{
+		PCHUsage = ModuleRules.PCHUsageMode.UseExplicitOrSharedPCHs;
+		
+		PublicIncludePaths.AddRange(
+			new string[] {
+				// ... add public include paths required here ...
+			}
+			);
+				
+		
+		PrivateIncludePaths.AddRange(
+			new string[] {
+				// ... add other private include paths required here ...
+			}
+			);
+			
+		
+		PublicDependencyModuleNames.AddRange(
+			new string[]
+			{
+				"Core",
+				// ... add other public dependencies that you statically link with here ...
+			}
+			);
+			
+		
+		PrivateDependencyModuleNames.AddRange(
+			new string[]
+			{
+				"CoreUObject",
+				"Engine",
+				"Slate",
+				"SlateCore", 
+				"Sequencer", 
+				"SequencerCore",
+				"GeometryCache", 
+				"LevelSequence", 
+				"LevelSequenceEditor",
+				"UnrealEd",
+				"MovieScene",
+				"InputCore",
+				"MeshDescription",
+				"StaticMeshDescription",
+				// ... add private dependencies that you statically link with here ...	
+			}
+			);
+		
+		
+		DynamicallyLoadedModuleNames.AddRange(
+			new string[]
+			{
+				// ... add any modules that your module loads dynamically here ...
+			}
+			);
+	}
+}
diff --git a/Engine/Plugins/Experimental/GeometryCacheLevelSequenceBaker/Source/GeometryCacheLevelSequenceBaker/Private/FGeometryCacheLevelSequenceBakerCustomization.cpp b/Engine/Plugins/Experimental/GeometryCacheLevelSequenceBaker/Source/GeometryCacheLevelSequenceBaker/Private/FGeometryCacheLevelSequenceBakerCustomization.cpp
new file mode 100644
--- /dev/null
+++ b/Engine/Plugins/Experimental/GeometryCacheLevelSequenceBaker/Source/GeometryCacheLevelSequenceBaker/Private/FGeometryCacheLevelSequenceBakerCustomization.cpp
@@ -0,0 +1,58 @@
+﻿// Copyright Epic Games, Inc. All Rights Reserved.
+#include "FGeometryCacheLevelSequenceBakerCustomization.h"
+
+#include "GeometryCacheLevelSequenceBaker.h"
+#include "GeometryCacheLevelSequenceBakerCommands.h"
+#include "Framework/MultiBox/MultiBoxExtender.h"
+#include "ISequencerModule.h"
+#include "ISequencer.h"
+#include "Modules/ModuleManager.h"
+#include "Toolkits/AssetEditorToolkit.h"
+
+void FGeometryCacheLevelSequenceBakerCustomization::RegisterSequencerCustomization(FSequencerCustomizationBuilder& Builder)
+{
+	WeakSequencer = Builder.GetSequencer().AsShared();
+
+	ISequencerModule& SequencerModule = FModuleManager::Get().LoadModuleChecked<ISequencerModule>("Sequencer");
+	
+	const FGeometryCacheLevelSequenceBakerCommands& Commands = FGeometryCacheLevelSequenceBakerCommands::Get();
+	
+	// Build the extender for the actions menu.
+	ActionsMenuCommandList = MakeShared<FUICommandList>().ToSharedPtr();
+	ActionsMenuCommandList->MapAction(
+		Commands.BakeGeometryCache,
+		FExecuteAction::CreateRaw( this, &FGeometryCacheLevelSequenceBakerCustomization::BakeGeometryCache),
+		FCanExecuteAction::CreateLambda( [] { return true; } ) );
+	
+	ActionsMenuExtender = MakeShared<FExtender>();
+	ActionsMenuExtender->AddMenuExtension(
+			"SequenceOptions", EExtensionHook::First, nullptr,
+			FMenuExtensionDelegate::CreateRaw(this, &FGeometryCacheLevelSequenceBakerCustomization::ExtendActionsMenu));
+
+	SequencerModule.GetActionsMenuExtensibilityManager()->AddExtender(ActionsMenuExtender);	
+}
+
+void FGeometryCacheLevelSequenceBakerCustomization::UnregisterSequencerCustomization()
+{
+	ISequencerModule& SequencerModule = FModuleManager::LoadModuleChecked<ISequencerModule>("Sequencer");
+	SequencerModule.GetActionsMenuExtensibilityManager()->RemoveExtender(ActionsMenuExtender);
+
+	ActionsMenuCommandList = nullptr;
+	WeakSequencer = nullptr;
+}
+
+void FGeometryCacheLevelSequenceBakerCustomization::ExtendActionsMenu(FMenuBuilder& MenuBuilder)
+{
+	MenuBuilder.PushCommandList(ActionsMenuCommandList.ToSharedRef());
+	{
+		const FGeometryCacheLevelSequenceBakerCommands& Commands = FGeometryCacheLevelSequenceBakerCommands::Get();
+		
+		MenuBuilder.AddMenuEntry(Commands.BakeGeometryCache);
+	}
+	MenuBuilder.PopCommandList();
+}
+
+void FGeometryCacheLevelSequenceBakerCustomization::BakeGeometryCache()
+{
+	FGeometryCacheLevelSequenceBaker::Bake(WeakSequencer.Pin().ToSharedRef());
+}
diff --git a/Engine/Plugins/Experimental/GeometryCacheLevelSequenceBaker/Source/GeometryCacheLevelSequenceBaker/Private/FGeometryCacheLevelSequenceBakerCustomization.h b/Engine/Plugins/Experimental/GeometryCacheLevelSequenceBaker/Source/GeometryCacheLevelSequenceBaker/Private/FGeometryCacheLevelSequenceBakerCustomization.h
new file mode 100644
--- /dev/null
+++ b/Engine/Plugins/Experimental/GeometryCacheLevelSequenceBaker/Source/GeometryCacheLevelSequenceBaker/Private/FGeometryCacheLevelSequenceBakerCustomization.h
@@ -0,0 +1,18 @@
+﻿// Copyright Epic Games, Inc. All Rights Reserved.
+#pragma once
+#include "SequencerCustomizationManager.h"
+
+class FGeometryCacheLevelSequenceBakerCustomization : public ISequencerCustomization
+{
+public:
+	void RegisterSequencerCustomization(FSequencerCustomizationBuilder& Builder) override;
+	void UnregisterSequencerCustomization() override;
+
+	void ExtendActionsMenu(FMenuBuilder& MenuBuilder);
+	void BakeGeometryCache();
+	
+	TWeakPtr<ISequencer> WeakSequencer;
+
+	TSharedPtr<FUICommandList> ActionsMenuCommandList;
+	TSharedPtr<FExtender> ActionsMenuExtender;
+};
diff --git a/Engine/Plugins/Experimental/GeometryCacheLevelSequenceBaker/Source/GeometryCacheLevelSequenceBaker/Private/GeometryCacheLevelSequenceBaker.cpp b/Engine/Plugins/Experimental/GeometryCacheLevelSequenceBaker/Source/GeometryCacheLevelSequenceBaker/Private/GeometryCacheLevelSequenceBaker.cpp
new file mode 100644
--- /dev/null
+++ b/Engine/Plugins/Experimental/GeometryCacheLevelSequenceBaker/Source/GeometryCacheLevelSequenceBaker/Private/GeometryCacheLevelSequenceBaker.cpp
@@ -0,0 +1,944 @@
+// Copyright Epic Games, Inc. All Rights Reserved.
+#include "GeometryCacheLevelSequenceBaker.h"
+
+#include "Components/SkeletalMeshComponent.h"
+#include "AssetToolsModule.h"
+#include "EditorDirectories.h"
+#include "GeometryCache.h"
+#include "GeometryCacheConstantTopologyWriter.h"
+#include "GeometryCacheLevelSequenceBakerModule.h"
+#include "GeometryCacheTrack.h"
+#include "GeometryCacheTrackStreamable.h"
+#include "MovieScene.h"
+#include "MovieSceneSequencePlayer.h"
+#include "AssetRegistry/AssetRegistryModule.h"
+#include "Dialogs/DlgPickAssetPath.h"
+#include "Framework/Notifications/NotificationManager.h"
+#include "MVVM/ViewModelPtr.h"
+#include "MVVM/Extensions/IObjectBindingExtension.h"
+#include "MVVM/Extensions/IOutlinerExtension.h"
+#include "MVVM/Selection/Selection.h"
+#include "MVVM/ViewModels/SequencerEditorViewModel.h"
+#include "Rendering/SkeletalMeshModel.h"
+#include "Rendering/SkeletalMeshRenderData.h"
+#include "Widgets/Notifications/SNotificationList.h"
+#include "Engine/SkinnedAsset.h"
+#include "Engine/SkeletalMesh.h"
+#include "Interfaces/IMainFrameModule.h"
+#include "PropertyEditorModule.h"
+#include "Modules/ModuleManager.h"
+#include "Framework/Application/SlateApplication.h"
+#include "IDetailsView.h"
+#include "SequencerSettings.h"
+#include "StaticMeshAttributes.h"
+#include "Animation/MeshDeformerGeometryReadback.h"
+#include "Async/ParallelFor.h"
+
+
+#define LOCTEXT_NAMESPACE "GeometryCacheLevelSequenceBaker"
+
+// At the moment there are meshes with MeshToImportVertexMap not matching their mesh description for some reason
+// Once the root cause is determined we can remove this
+#define USE_MESH_TO_IMPORT_VERTEX_MAP 0
+
+const int32 FGeometryCacheLevelSequenceBaker::LODIndexToBake = 0;
+const float FGeometryCacheLevelSequenceBaker::TotalAmountOfWork = 1.0f;
+const float FGeometryCacheLevelSequenceBaker::AmountOfWorkGatherStage = TotalAmountOfWork / 2.0f;
+const float FGeometryCacheLevelSequenceBaker::AmountOfWorkBakeStage = TotalAmountOfWork - AmountOfWorkGatherStage;
+
+class SLevelSequenceGeometryCacheBakerOptionWindow : public SCompoundWidget
+{
+public:
+	SLATE_BEGIN_ARGS(SLevelSequenceGeometryCacheBakerOptionWindow) {}
+		SLATE_ARGUMENT( TSharedPtr<SWindow>, WidgetWindow )
+		SLATE_ARGUMENT( ULevelSequenceGeometryCacheBakerOption*, Option)
+	SLATE_END_ARGS()
+
+	void Construct(const FArguments& InArgs)
+	{
+		WidgetWindow = InArgs._WidgetWindow;
+		Option = InArgs._Option;
+
+
+		FPropertyEditorModule& PropertyEditorModule = FModuleManager::GetModuleChecked<FPropertyEditorModule>("PropertyEditor");
+		FDetailsViewArgs DetailsViewArgs;
+		DetailsViewArgs.bAllowSearch = false;
+		DetailsViewArgs.NameAreaSettings = FDetailsViewArgs::HideNameArea;
+		DetailsView = PropertyEditorModule.CreateDetailView(DetailsViewArgs);
+		DetailsView->SetObject(Option);
+		
+		this->ChildSlot
+		[
+			SNew(SVerticalBox)
+			+SVerticalBox::Slot()
+			.FillHeight(1.0f)
+			[
+				DetailsView.ToSharedRef()	
+			]
+			+SVerticalBox::Slot()
+			.AutoHeight()
+			[
+				SNew(SHorizontalBox)
+				+SHorizontalBox::Slot()
+				[
+					SAssignNew(BakeButton, SButton)
+					.HAlign(HAlign_Center)
+					.Text(LOCTEXT("LevelSequenceGeometryCacheBakerOptionsWindow_Bake", "Bake"))
+					.OnClicked(this, &SLevelSequenceGeometryCacheBakerOptionWindow::OnBake)
+				]
+				+SHorizontalBox::Slot()
+				[
+					SAssignNew(CancelButton, SButton)
+					.HAlign(HAlign_Center)
+					.Text(LOCTEXT("LevelSequenceGeometryCacheBakerOptionsWindow_Cancel", "Cancel"))
+					.OnClicked(this, &SLevelSequenceGeometryCacheBakerOptionWindow::OnCancel)
+				]
+			]
+		];
+	}
+
+	bool SupportsKeyboardFocus() const override { return true; }
+	
+	FReply OnBake()
+	{
+		Option->bShouldBake = true;
+		if ( WidgetWindow.IsValid() )
+		{
+			WidgetWindow.Pin()->RequestDestroyWindow();
+		}
+		return FReply::Handled();
+	}
+
+	FReply OnCancel()
+	{
+		if ( WidgetWindow.IsValid() )
+		{
+			WidgetWindow.Pin()->RequestDestroyWindow();
+		}
+		return FReply::Handled();
+	}
+
+	FReply OnKeyDown(const FGeometry& MyGeometry, const FKeyEvent& InKeyEvent)
+	{
+		if( InKeyEvent.GetKey() == EKeys::Escape )
+		{
+			return OnCancel();
+		}
+
+		return FReply::Unhandled();
+	}
+
+	TWeakPtr<SWindow> WidgetWindow;
+	TObjectPtr<ULevelSequenceGeometryCacheBakerOption> Option;
+	TSharedPtr<IDetailsView> DetailsView;
+	TSharedPtr<SButton> BakeButton;
+	TSharedPtr<SButton> CancelButton;
+};
+
+
+
+void FGeometryCacheLevelSequenceBaker::Bake(TSharedRef<ISequencer> InSequencer)
+{
+	TSharedPtr<ISequencer> Sequencer = InSequencer;
+	UMovieSceneSequence* MovieSceneSequence = Sequencer->GetFocusedMovieSceneSequence();
+	UMovieScene* MovieScene = MovieSceneSequence->GetMovieScene();
+	
+	TArray<FGuid> Bindings = GetBindingsToBake(InSequencer);
+	if (Bindings.Num() == 0 )
+	{
+		return;
+	}
+	
+	FString PackageName;
+	FString AssetName;
+
+	if (!GetGeometryCacheAssetPathFromUser(PackageName, AssetName))
+	{
+		return;
+	}
+
+	TSharedRef<SWindow> Window = SNew(SWindow)
+		.Title(LOCTEXT("LevelSequenceGeometryCacheBakerOptionTitle", "Bake Geometry Cache Options"))
+		.SizingRule(ESizingRule::Autosized);
+
+	ULevelSequenceGeometryCacheBakerOption* Option = NewObject<ULevelSequenceGeometryCacheBakerOption>();
+	
+	TSharedPtr<SLevelSequenceGeometryCacheBakerOptionWindow> OptionWindow;
+	Window->SetContent
+	(
+		SAssignNew(OptionWindow, SLevelSequenceGeometryCacheBakerOptionWindow)
+		.WidgetWindow(Window)
+		.Option(Option)
+	);
+	
+	GEditor->EditorAddModalWindow(Window);
+
+	if (Option->bShouldBake == false)
+	{
+		return;
+	}
+
+	const TRange<FFrameNumber> PlaybackRange = MovieScene->GetPlaybackRange();
+	const FFrameRate Resolution = MovieScene->GetTickResolution();
+	const FFrameRate DisplayRate = MovieScene->GetDisplayRate();
+
+	const FFrameTime Interval = FFrameRate::TransformTime(1, DisplayRate, Resolution);
+	FFrameTime StartFrame = UE::MovieScene::DiscreteInclusiveLower(PlaybackRange);
+	FFrameTime EndFrame = UE::MovieScene::DiscreteExclusiveUpper(PlaybackRange);
+
+	int32 NumFrames = FMath::FloorToInt32((EndFrame.AsDecimal() - StartFrame.AsDecimal()) / Interval.AsDecimal());
+	int32 NumSamplesPerFrame = Option->NumSamplesPerFrame;
+	int32 TotalSamples = NumFrames * NumSamplesPerFrame;
+	// Also capture the last frame since geometry cache use the last frame's time to compute total duration
+	TotalSamples += 1;
+
+	
+
+	// Setup bake task
+	FGeometryCacheLevelSequenceBaker& Baker = Get();
+	check(Baker.CurrentBakeTask == nullptr);
+	Baker.CurrentBakeTask =  MakeUnique<FBakeTask>();
+
+	TUniquePtr<FBakeTask>& BakeTask	= Baker.CurrentBakeTask;
+	BakeTask->PackageName = PackageName;
+	BakeTask->AssetName = AssetName;
+
+	BakeTask->Sequencer = Sequencer;
+	BakeTask->Bindings = Bindings;
+	BakeTask->StartFrame = StartFrame;
+	BakeTask->EndFrame = EndFrame;
+	
+	BakeTask->NumSamples = TotalSamples;
+	BakeTask->CurrentSampleIndex = 0;
+	BakeTask->SamplesPerSecond = static_cast<float>(NumSamplesPerFrame * DisplayRate.AsDecimal());
+
+	BakeTask->SlowTask = MakeUnique<FScopedSlowTask>(1.0f, LOCTEXT("BakeGeometryCacheSlowTask", "Baking Geometry Cache..."));
+	BakeTask->SlowTask->MakeDialog(true);
+	
+	double DeltaTime = DisplayRate.AsInterval() / NumSamplesPerFrame;	
+	BakeTask->FixedDeltaTimeScope = MakeUnique<FEngineFixedDeltaTimeScope>(DeltaTime);
+	BakeTask->SequencerStateScope = MakeUnique<FSequencerSettingScope>(Sequencer);
+	BakeTask->ConsoleVariableOverrideScope = MakeUnique<FConsoleVariableOverrideScope>();
+}
+
+TArray<FGuid> FGeometryCacheLevelSequenceBaker::GetBindingsToBake(TSharedRef<ISequencer> InSequencer)
+{
+	TSharedPtr<ISequencer> Sequencer = InSequencer;
+	UMovieSceneSequence* MovieSceneSequence = Sequencer->GetFocusedMovieSceneSequence();
+	UMovieScene* MovieScene = MovieSceneSequence->GetMovieScene();
+	
+	using namespace UE::Sequencer;
+
+	// Bake only selected bindings if there are any, otherwise bake every binding
+	TArray<FGuid> SelectedBindings;
+	TSharedPtr<FSequencerEditorViewModel> EditorViewModel = Sequencer->GetViewModel();
+	for (FViewModelPtr Node : EditorViewModel->GetSelection()->Outliner)
+	{
+		if (TViewModelPtr<IObjectBindingExtension> ObjectBindingNode = Node.ImplicitCast())
+		{
+			SelectedBindings.Add(ObjectBindingNode->GetObjectGuid());
+		}
+	}
+
+	TArray<FGuid> CandidateBindings;
+	for (const FMovieSceneBinding& MovieSceneBinding : MovieScene->GetBindings())
+	{
+		// If there are specific bindings to export, export those only
+		if (SelectedBindings.Num() != 0 && !SelectedBindings.Contains(MovieSceneBinding.GetObjectGuid()))
+		{
+			continue;
+		}
+
+		CandidateBindings.Add(MovieSceneBinding.GetObjectGuid());
+	}
+	
+	// Skip a child binding if its parent binding is already getting baked
+	TArray<FMovieScenePossessable*> Possessables;
+	TArray<FMovieSceneSpawnable*> Spawnables;
+	for (FGuid Binding : CandidateBindings)
+	{
+		if (FMovieScenePossessable* Possessable = MovieScene->FindPossessable(Binding))
+		{
+			Possessables.Add(Possessable);
+		}
+		else if (FMovieSceneSpawnable* Spawnable = MovieScene->FindSpawnable(Binding))
+		{
+			Spawnables.Add(Spawnable);
+		}
+	}
+
+	TArray<FGuid> FinalizedBindings;
+	for (FMovieSceneSpawnable* Spawnable : Spawnables)
+	{
+		FinalizedBindings.Add(Spawnable->GetGuid());
+	}
+	
+	for (FMovieScenePossessable* Possessable : Possessables)
+	{
+		check(Possessable);
+		// Skip possessables that belong to a spawnable that we are already baking
+		for (FMovieSceneSpawnable* Spawnable : Spawnables)
+		{
+			check(Spawnable);
+			
+			if (Spawnable->GetChildPossessables().Contains(Possessable->GetGuid()))
+			{
+				continue;
+			}
+		}
+
+		// Skip possessables that belong to a Possessable that we are already baking
+		TArray<FMovieScenePossessable*> PossessableParents;
+		FMovieScenePossessable* WorkItem = Possessable;
+		while (WorkItem)
+		{
+			WorkItem = MovieScene->FindPossessable(WorkItem->GetParent());
+			if (WorkItem)
+			{
+				PossessableParents.Add(WorkItem);
+			}
+		}
+
+		bool bSkip = false;
+		for (FMovieScenePossessable* Parent : PossessableParents)
+		{
+			if (Possessables.Contains(Parent))
+			{
+				bSkip = true;
+				break;
+			}
+		}
+
+		if (bSkip)
+		{
+			continue;
+		}
+
+		// Include this possessable since it does not belong to any other bindings that will be baked
+		FinalizedBindings.Add(Possessable->GetGuid());
+	}
+
+
+	return FinalizedBindings;
+}
+
+
+bool FGeometryCacheLevelSequenceBaker::GetGeometryCacheAssetPathFromUser(FString& OutPackageName, FString& OutAssetName)
+{
+	FString NewNameSuggestion = FString(TEXT("NewGeometryCache"));
+	FString DefaultPath;
+	const FString DefaultDirectory = FEditorDirectories::Get().GetLastDirectory(ELastDirectory::NEW_ASSET);
+	FPackageName::TryConvertFilenameToLongPackageName(DefaultDirectory, DefaultPath);
+
+	if (DefaultPath.IsEmpty())
+	{
+		DefaultPath = TEXT("/Game/GeometryCaches");
+	}
+
+	FString PackageNameSuggestion = DefaultPath / NewNameSuggestion;
+	FString Name;
+	FAssetToolsModule& AssetToolsModule = FModuleManager::LoadModuleChecked<FAssetToolsModule>("AssetTools");
+	AssetToolsModule.Get().CreateUniqueAssetName(PackageNameSuggestion, TEXT(""), PackageNameSuggestion, Name);
+
+	// Decide where to create the geo cache asset
+	TSharedPtr<SDlgPickAssetPath> PickAssetPathWidget =
+				SNew(SDlgPickAssetPath)
+				.Title(LOCTEXT("BakeGeometryCachePickName", "Choose New Geometry Cache Location"))
+				.DefaultAssetPath(FText::FromString(PackageNameSuggestion));
+
+	if (PickAssetPathWidget->ShowModal() == EAppReturnType::Ok)
+	{
+		// Get the full name of where we want to create the mesh asset.
+		OutPackageName = PickAssetPathWidget->GetFullAssetPath().ToString();
+		OutAssetName = FPackageName::GetLongPackageAssetName(OutPackageName);
+
+		// Check if the user inputed a valid asset name, if they did not, give it the generated default name
+		if (OutAssetName.IsEmpty())
+		{
+			// Use the defaults that were already generated.
+			OutPackageName = PackageNameSuggestion;
+			OutAssetName = *Name;
+		}
+
+		return true;
+	}
+
+	return false;
+}
+
+
+void FGeometryCacheLevelSequenceBaker::Tick(float DeltaTime)
+{
+	if (CurrentBakeTask->Stage == EStage::Gather)
+	{
+		if (CurrentBakeTask->CurrentSampleIndex == 0)
+		{
+			CurrentBakeTask->PlaySequencer();
+			return;
+		}
+
+		if (!CurrentBakeTask->IsSequencerPlaying())
+		{
+			check(CurrentBakeTask->CurrentSampleIndex == CurrentBakeTask->NumSamples);
+			
+			CurrentBakeTask->CurrentSampleIndex = 0;
+			CurrentBakeTask->GatherStageComponentSettingScopes.Reset();
+			SetupComponentBakeTasks();
+		
+			CurrentBakeTask->Stage = EStage::RequestReadback;	
+		}
+	}
+		
+
+
+	if (CurrentBakeTask->Stage == EStage::RequestReadback)
+	{
+		if (CurrentBakeTask->CurrentSampleIndex == 0)
+		{
+			CurrentBakeTask->PlaySequencer();
+			return;
+		}
+		
+		if (!CurrentBakeTask->IsSequencerPlaying())
+		{
+			check(CurrentBakeTask->CurrentSampleIndex == CurrentBakeTask->NumSamples);
+			CurrentBakeTask->Stage = EStage::WriteToAsset;
+		}
+	}
+
+	if (CurrentBakeTask->Stage == EStage::WriteToAsset)
+	{
+		if (!CurrentBakeTask->SlowTask->ShouldCancel())
+		{
+			if (CurrentBakeTask->NumComponentTasksPending > 0)
+			{
+				return;
+			}
+			
+			WriteToAsset();
+		}
+
+		// Make sure to wait for ongoing readback callbacks to finish
+		FScopeLock BakeTaskLifeTimeLock(&CurrentBakeTaskLifeTimeCriticalSection);
+		CurrentBakeTask = nullptr;
+	}
+}
+
+void FGeometryCacheLevelSequenceBaker::OnSequencerEvaluated()
+{
+	if (CurrentBakeTask->Stage == EStage::Gather)
+	{
+		if (CurrentBakeTask->CurrentSampleIndex < CurrentBakeTask->NumSamples)
+		{
+			Gather();
+			CurrentBakeTask->CurrentSampleIndex++;
+		}
+		return;
+	}
+
+	if (CurrentBakeTask->Stage == EStage::RequestReadback)
+	{
+		if (CurrentBakeTask->CurrentSampleIndex < CurrentBakeTask->NumSamples)
+		{
+			RequestReadback();	
+			CurrentBakeTask->CurrentSampleIndex++;
+		}
+
+		return;
+	}	
+}
+
+bool FGeometryCacheLevelSequenceBaker::IsTickable() const
+{
+	return CurrentBakeTask.IsValid();
+}
+
+FGeometryCacheLevelSequenceBaker& FGeometryCacheLevelSequenceBaker::Get()
+{
+	static FGeometryCacheLevelSequenceBaker Instance;
+	return Instance;
+}
+
+void FGeometryCacheLevelSequenceBaker::SetupComponentBakeTasks()
+{
+	// Create readback tasks for each skeletal mesh component
+	TArray<FComponentTask>& ComponentTasks = CurrentBakeTask->ComponentTasks;
+	
+	for (const TPair<FGuid, TMap<FName, FComponentInfo>>& BindingToComponentInfo : CurrentBakeTask->BindingToComponentInfoMap)
+	{
+		for (const TPair<FName, FComponentInfo>& ComponentInfo : BindingToComponentInfo.Value)
+		{
+			FComponentTask& ComponentTask = ComponentTasks.AddDefaulted_GetRef();
+			ComponentTask.Binding = BindingToComponentInfo.Key;
+			ComponentTask.ComponentInfo = ComponentInfo.Value;
+			ComponentTask.NumSamplesPending = CurrentBakeTask->NumSamples;
+			ComponentTask.GeometrySamples.Init({}, ComponentTask.NumSamplesPending);
+		
+			const FSkeletalMeshLODModel& LODModel = ComponentTask.ComponentInfo.SkeletalMeshAsset->GetImportedModel()->LODModels[LODIndexToBake];
+			for (FFrameData& Sample : ComponentTask.GeometrySamples)
+			{
+				Sample.Positions.Init({}, LODModel.NumVertices);
+				Sample.Normals.Init({}, LODModel.NumVertices);
+				Sample.TangentsX.Init({}, LODModel.NumVertices);
+			}
+			
+			ComponentTask.VisibilitySamples.Init({}, ComponentTask.NumSamplesPending);
+		}
+	}
+	
+	CurrentBakeTask->NumComponentTasksPending = CurrentBakeTask->ComponentTasks.Num();	
+}
+
+void FGeometryCacheLevelSequenceBaker::Gather()
+{
+	TSharedPtr<ISequencer> Sequencer = CurrentBakeTask->Sequencer.Pin();
+	IMovieScenePlayer* MovieScenePlayer = Sequencer.Get();
+	FMovieSceneSequenceIDRef SequenceID = Sequencer->GetFocusedTemplateID();
+
+	for (FGuid Binding : CurrentBakeTask->Bindings)
+	{
+		TMap<FName, FComponentInfo>& ComponentInfosRef = CurrentBakeTask->BindingToComponentInfoMap.FindOrAdd(Binding);
+		TArrayView<TWeakObjectPtr<>> RuntimeObjects = MovieScenePlayer->FindBoundObjects(Binding, SequenceID);
+		if (RuntimeObjects.Num() > 0)
+		{
+			// At the moment only support one object per binding
+			TWeakObjectPtr<UObject> RuntimeObject = RuntimeObjects[0];
+
+			TArray<USkeletalMeshComponent*> SkeletalMeshComponents;
+			if (RuntimeObject.IsValid())
+			{
+				if (AActor* Actor = Cast<AActor>(RuntimeObject.Get()))
+				{
+					// Bake all the skeletal mesh components for this actor
+					TArray<USkeletalMeshComponent*> ActorSkeletalMeshComponents;
+					Actor->GetComponents(ActorSkeletalMeshComponents);
+					for (USkeletalMeshComponent* SkeletalMeshComponent : ActorSkeletalMeshComponents)
+					{
+						SkeletalMeshComponents.Add(SkeletalMeshComponent);
+					}
+				}
+				else if (USkeletalMeshComponent* SkeletalMeshComponent = Cast<USkeletalMeshComponent>(RuntimeObject.Get()))
+				{
+					SkeletalMeshComponents.Add(SkeletalMeshComponent);
+				}
+			}
+
+			for (USkeletalMeshComponent* SkeletalMeshComponent : SkeletalMeshComponents)
+			{
+				if (USkeletalMesh* SkeletalMesh = SkeletalMeshComponent->GetSkeletalMeshAsset())
+				{
+					FName Name = SkeletalMeshComponent->GetFName();
+					if (FComponentInfo* ComponentInfo = ComponentInfosRef.Find(Name))
+					{
+						// At the moment we don't support changing mesh assets
+						if (!ensure(ComponentInfo->SkeletalMeshAsset == SkeletalMesh))
+						{
+							FNotificationInfo ErrorToast(
+								LOCTEXT("FailToBakeGeometryCache_Title",
+								        "Bake to Geometry Cache failed"));
+							ErrorToast.ExpireDuration = 10.0f;
+							ErrorToast.bFireAndForget = true;
+							ErrorToast.Image = FCoreStyle::Get().GetBrush(TEXT("MessageLog.Error"));
+							ErrorToast.SubText = LOCTEXT("FailToBakeGeometryCache_MismatchedMeshAsset",
+							                             "Changing Skeletal Mesh Asset during Playback is not Supported");
+							FSlateNotificationManager::Get().AddNotification(ErrorToast);
+
+							// clear the task and abort
+							CurrentBakeTask = nullptr;
+							return;
+						}
+					}
+					else
+					{
+						ComponentInfosRef.Add(Name, {Name, SkeletalMesh});
+					}
+
+					ComponentInfosRef.Add(Name, {Name, SkeletalMesh});
+
+					// Try to test-run deformers on these skeletal mesh component to make sure deformer shaders are compiled and ready
+					TUniquePtr<FSkeletalMeshComponentSettingScope>& ComponentSettingScope = CurrentBakeTask->GatherStageComponentSettingScopes.
+						FindOrAdd(SkeletalMeshComponent);
+					if (!ComponentSettingScope.IsValid())
+					{
+						ComponentSettingScope = MakeUnique<FSkeletalMeshComponentSettingScope>(SkeletalMeshComponent);
+					}
+				}
+			}
+		}
+	}
+
+	CurrentBakeTask->UpdateGatherProgress();
+}
+
+void FGeometryCacheLevelSequenceBaker::RequestReadback()
+{
+	TSharedPtr<ISequencer> Sequencer = CurrentBakeTask->Sequencer.Pin();
+	IMovieScenePlayer* MovieScenePlayer = Sequencer.Get();
+	FMovieSceneSequenceIDRef SequenceID = Sequencer->GetFocusedTemplateID();
+
+	for (int32 TaskIndex = 0;  TaskIndex < CurrentBakeTask->ComponentTasks.Num(); ++TaskIndex)
+	{
+		FComponentTask& ComponentTask = CurrentBakeTask->ComponentTasks[TaskIndex];
+		
+		bool bReadbackRequested = false;
+
+		USkeletalMeshComponent* Component = nullptr;
+		
+		TArrayView<TWeakObjectPtr<>> RuntimeObjects = MovieScenePlayer->FindBoundObjects(ComponentTask.Binding, SequenceID);
+		if (RuntimeObjects.Num() > 0 && RuntimeObjects[0].IsValid())
+		{
+			// At the moment only support one object per binding
+			TWeakObjectPtr<UObject> RuntimeObject = RuntimeObjects[0];
+
+			if (AActor* Actor = Cast<AActor>(RuntimeObject.Get()))
+			{
+				// Bake all the skeletal mesh components for this actor
+				TArray<USkeletalMeshComponent*> ActorSkeletalMeshComponents;
+				Actor->GetComponents(ActorSkeletalMeshComponents);
+
+				for (USkeletalMeshComponent* SkeletalMeshComponent : ActorSkeletalMeshComponents)
+				{
+					if (SkeletalMeshComponent->GetFName() == ComponentTask.ComponentInfo.Name)
+					{
+						Component = SkeletalMeshComponent;
+					}
+				}
+			}
+			else if (USkeletalMeshComponent* SkeletalMeshComponent = Cast<USkeletalMeshComponent>(RuntimeObject.Get()))
+			{
+				if (SkeletalMeshComponent->GetFName() == ComponentTask.ComponentInfo.Name)
+				{
+					Component = SkeletalMeshComponent;
+				}
+			}
+		}
+
+		if (Component)
+		{
+			if (ComponentTask.ComponentSettingScope == nullptr || ComponentTask.ComponentSettingScope->Component != Component)
+			{
+				ComponentTask.ComponentSettingScope = MakeUnique<FSkeletalMeshComponentSettingScope>(Component);
+			}
+
+			FTransform Transform = Component->GetComponentTransform();
+			TUniquePtr<FMeshDeformerGeometryReadbackRequest> Request = MakeUnique<FMeshDeformerGeometryReadbackRequest>();
+			Request->MeshDescriptionCallback_AnyThread =
+				[
+					Transform,
+					SampleIndex = CurrentBakeTask->CurrentSampleIndex,
+					TaskIndex
+				](const FMeshDescription& MeshDescription)
+				{
+					// Making sure current bake task is alive when this callback is invoked on a worker thread
+					FScopeLock BakeTaskLifeTimeLock(&Get().CurrentBakeTaskLifeTimeCriticalSection);
+					TUniquePtr<FBakeTask>& BakeTask = Get().CurrentBakeTask; 
+					if (BakeTask == nullptr)
+					{
+						return;
+					}
+
+					FComponentTask& ComponentTask = BakeTask->ComponentTasks[TaskIndex];
+
+					bool bMeshAvailable = false;
+					if (!MeshDescription.IsEmpty())
+					{
+						bMeshAvailable = true;
+						
+						const USkeletalMesh* MeshAsset = ComponentTask.ComponentInfo.SkeletalMeshAsset;
+						const FSkeletalMeshLODModel& LODModel = MeshAsset->GetImportedModel()->LODModels[LODIndexToBake];
+#if USE_MESH_TO_IMPORT_VERTEX_MAP
+						const TArray<int32>& RenderToImportedMap = LODModel.MeshToImportVertexMap;
+#else
+						const TArray<uint32>& RenderToImportedMap = LODModel.GetRawPointIndices();
+#endif
+						
+						TMeshAttributesRef<FVertexID, const FVector3f> PositionAttribute =
+							MeshDescription.VertexAttributes().GetAttributesRef<FVector3f>(MeshAttribute::Vertex::Position);
+						TMeshAttributesRef<FVertexInstanceID, const FVector3f> NormalAttribute =
+							MeshDescription.VertexInstanceAttributes().GetAttributesRef<FVector3f>(MeshAttribute::VertexInstance::Normal);
+						TMeshAttributesRef<FVertexInstanceID, const FVector3f> TangentAttribute =
+							MeshDescription.VertexInstanceAttributes().GetAttributesRef<FVector3f>(MeshAttribute::VertexInstance::Tangent);
+
+						ParallelFor(LODModel.NumVertices,
+						[
+							&RenderToImportedMap,
+							&MeshDescription,
+							&PositionAttribute,
+							&NormalAttribute,
+							&TangentAttribute,
+							&FrameData = ComponentTask.GeometrySamples[SampleIndex],
+							&Transform
+						](int32 Index)
+						{
+							FVertexID ID = RenderToImportedMap[Index];
+
+							FrameData.Positions[Index] = FVector3f(Transform.TransformPosition(FVector(PositionAttribute[ID])));
+
+							const TArrayView<const FVertexInstanceID> VertexInstances = MeshDescription.
+								GetVertexVertexInstanceIDs(ID);
+
+							// Some vertices for some reason have no vertex instances, dangling verts?
+							if (VertexInstances.Num() == 0)
+							{
+								return;
+							}
+							
+							// Normal
+							FVector3f Normal = NormalAttribute[VertexInstances[0]];
+
+							FrameData.Normals[Index] = FVector3f(Transform.TransformVector(FVector(Normal)));
+
+							// TangentX
+							FVector3f TangentX = TangentAttribute[VertexInstances[0]];
+							FrameData.TangentsX[Index] = FVector3f(Transform.TransformVector(FVector(TangentX)));
+						});	
+					}
+
+					OnReadbackResultConfirmed(ComponentTask, SampleIndex, bMeshAvailable);	
+				};
+				
+			Component->RequestReadbackRenderGeometry(MoveTemp(Request));
+			
+			bReadbackRequested = true;
+		}
+
+		if (!bReadbackRequested)
+		{
+			constexpr bool bMeshAvailable = false;
+			OnReadbackResultConfirmed(ComponentTask, CurrentBakeTask->CurrentSampleIndex, bMeshAvailable);		
+		}
+	}
+}
+
+void FGeometryCacheLevelSequenceBaker::WriteToAsset()
+{
+	UPackage* Package = CreatePackage(*CurrentBakeTask->PackageName);
+	check(Package);
+	UGeometryCache* GeometryCache = NewObject<UGeometryCache>(Package, *CurrentBakeTask->AssetName, RF_Public | RF_Standalone);
+	
+	using UE::GeometryCacheHelpers::FGeometryCacheConstantTopologyWriter;
+	using UE::GeometryCacheHelpers::AddTrackWriterFromSkinnedAsset;
+	using FTrackWriter = FGeometryCacheConstantTopologyWriter::FTrackWriter;
+	FGeometryCacheConstantTopologyWriter::FConfig Config;
+	Config.FPS = CurrentBakeTask->SamplesPerSecond;
+	FGeometryCacheConstantTopologyWriter Writer(*GeometryCache, Config);
+
+	for (FComponentTask& Task : CurrentBakeTask->ComponentTasks)
+	{
+
+		int32 FirstVisibleFrameIndex = INDEX_NONE;
+		for (const FVisibilitySample& VisibilitySample : Task.VisibilitySamples)
+		{
+			if (VisibilitySample.bVisible && FirstVisibleFrameIndex == INDEX_NONE)
+			{
+				FirstVisibleFrameIndex = VisibilitySample.FrameIndex;
+			}
+		}
+
+		if (FirstVisibleFrameIndex == INDEX_NONE)
+		{
+			continue;
+		}
+		
+		// Copy previous frame if the component isn't visible at this frame so that
+        // frame interpolation works correctly
+		for (const FVisibilitySample& VisibilitySample : Task.VisibilitySamples)
+		{	
+			if (!VisibilitySample.bVisible)
+			{
+				if (VisibilitySample.FrameIndex > 0)
+				{
+					if (Task.VisibilitySamples[VisibilitySample.FrameIndex-1].bVisible)
+					{
+						Task.GeometrySamples[VisibilitySample.FrameIndex] = Task.GeometrySamples[VisibilitySample.FrameIndex - 1];	
+					}
+				}
+			}
+		}
+
+		// Make sure the first frame has valid data, even if it is invisible, so that the geo cache preprocessor can use it to for sanity checks
+		if (!Task.VisibilitySamples[0].bVisible)
+		{
+			Task.GeometrySamples[0] = Task.GeometrySamples[Task.VisibilitySamples[FirstVisibleFrameIndex].FrameIndex];
+		}
+
+		int32 TrackIndex = AddTrackWriterFromSkinnedAsset(Writer, *(Task.ComponentInfo.SkeletalMeshAsset));
+
+		FTrackWriter& TrackWriter = Writer.GetTrackWriter(TrackIndex);
+
+		TrackWriter.WriteAndClose(Task.GeometrySamples, Task.VisibilitySamples);
+		Task.GeometrySamples.Reset();
+	}
+
+	GeometryCache->MarkPackageDirty();
+
+	// Notify asset registry of new asset
+	FAssetRegistryModule::AssetCreated(GeometryCache);
+
+	// Display notification so users can quickly access the mesh
+	if (GIsEditor)
+	{
+		FNotificationInfo Info(FText::Format(
+			LOCTEXT("GeometryCacheBaked", "Successfully Baked to Geometry Cache"), FText::FromString(GeometryCache->GetName())));
+		Info.ExpireDuration = 15.0f;
+		Info.bUseLargeFont = false;
+		Info.Hyperlink = FSimpleDelegate::CreateLambda([=]()
+		{
+			GEditor->GetEditorSubsystem<UAssetEditorSubsystem>()->OpenEditorForAssets(TArray<UObject*>({GeometryCache}));
+		});
+		Info.HyperlinkText = FText::Format(LOCTEXT("OpenNewGeometryCacheHyperlink", "Open {0}"), FText::FromString(GeometryCache->GetName()));
+		TSharedPtr<SNotificationItem> Notification = FSlateNotificationManager::Get().AddNotification(Info);
+		if (Notification.IsValid())
+		{
+			Notification->SetCompletionState(SNotificationItem::CS_Success);
+		}
+	}
+}
+
+FGeometryCacheLevelSequenceBaker::FSkeletalMeshComponentSettingScope::FSkeletalMeshComponentSettingScope(USkeletalMeshComponent* InComponent)
+	: Component(InComponent)
+{
+
+	bPreviousAlwaysUseMeshDeformer = Component->GetAlwaysUseMeshDeformer();
+	Component->SetAlwaysUseMeshDeformer(true);
+
+	bPreviousUpdateAnimationInEditor = Component->GetUpdateAnimationInEditor();
+	Component->SetUpdateAnimationInEditor(true);
+
+	PreviousForcedLOD = Component->GetForcedLOD();
+	Component->SetForcedLOD(LODIndexToBake+1);
+}
+
+FGeometryCacheLevelSequenceBaker::FSkeletalMeshComponentSettingScope::~FSkeletalMeshComponentSettingScope()
+{
+	if (Component.IsValid())
+	{
+		Component->SetAlwaysUseMeshDeformer(bPreviousAlwaysUseMeshDeformer);
+		Component->SetUpdateAnimationInEditor(bPreviousUpdateAnimationInEditor);
+		Component->SetForcedLOD(PreviousForcedLOD);
+	}
+}
+
+void FGeometryCacheLevelSequenceBaker::OnReadbackResultConfirmed(FComponentTask& ComponentTask, int32 SampleIndex, bool bMeshAvailable)
+{
+	ComponentTask.VisibilitySamples[SampleIndex].FrameIndex = SampleIndex;
+	ComponentTask.VisibilitySamples[SampleIndex].bVisible = bMeshAvailable;
+
+	if (IsInGameThread())
+	{
+		Get().CurrentBakeTask->UpdateBakeProgress();
+	}
+	else
+	{
+		FFunctionGraphTask::CreateAndDispatchWhenReady([]()
+		{
+			Get().CurrentBakeTask->UpdateBakeProgress();
+		}, TStatId(), NULL, ENamedThreads::GameThread);
+	}
+
+	if (--ComponentTask.NumSamplesPending == 0)
+	{
+		--Get().CurrentBakeTask->NumComponentTasksPending;
+	}
+}
+
+FGeometryCacheLevelSequenceBaker::FEngineFixedDeltaTimeScope::FEngineFixedDeltaTimeScope(double NewFixedDeltaTime)
+{
+	bPreviousUseFixedDeltaTime = FApp::UseFixedTimeStep();
+	FApp::SetUseFixedTimeStep(true);
+	
+	PreviousFixedDeltaTime = FApp::GetFixedDeltaTime();
+	FApp::SetFixedDeltaTime(NewFixedDeltaTime);
+}
+
+FGeometryCacheLevelSequenceBaker::FEngineFixedDeltaTimeScope::~FEngineFixedDeltaTimeScope()
+{
+	FApp::SetUseFixedTimeStep(bPreviousUseFixedDeltaTime);
+	FApp::SetFixedDeltaTime(PreviousFixedDeltaTime);
+}
+
+FGeometryCacheLevelSequenceBaker::FSequencerSettingScope::FSequencerSettingScope(TSharedPtr<ISequencer> InSequencer)
+{
+	Sequencer = InSequencer;
+	PreviousGlobalTime = InSequencer->GetGlobalTime();
+	PreviousLoopMode = InSequencer->GetSequencerSettings()->GetLoopMode();
+
+	InSequencer->GetSequencerSettings()->SetLoopMode(ESequencerLoopMode::SLM_NoLoop);
+
+	OnGlobalTimeChangedDelegateHandle = InSequencer->OnGlobalTimeChanged().AddRaw(&Get(), &FGeometryCacheLevelSequenceBaker::OnSequencerEvaluated);
+}
+
+FGeometryCacheLevelSequenceBaker::FSequencerSettingScope::~FSequencerSettingScope()
+{
+	if (ensure(Sequencer.Pin()))
+	{
+		Sequencer.Pin()->OnGlobalTimeChanged().Remove(OnGlobalTimeChangedDelegateHandle);
+		Sequencer.Pin()->GetSequencerSettings()->SetLoopMode(PreviousLoopMode);
+		Sequencer.Pin()->SetGlobalTime(PreviousGlobalTime.Time);
+		
+	}
+}
+
+FGeometryCacheLevelSequenceBaker::FConsoleVariableOverrideScope::FConsoleVariableOverrideScope()
+{
+	// Needed for metahuman which uses LODSyncComponent
+	if (IConsoleVariable* CVarForceLOD = IConsoleManager::Get().FindConsoleVariable(TEXT("r.ForceLOD")))
+	{
+		PreviousForceLOD = CVarForceLOD->GetInt();
+		CVarForceLOD->Set(0);
+	}
+
+
+	if (IConsoleVariable* CVarSkeletalMeshLODBias = IConsoleManager::Get().FindConsoleVariable(TEXT("r.SkeletalMeshLODBias")))
+	{
+		PreviousSkeletalMeshLODBias = CVarSkeletalMeshLODBias->GetInt();
+		CVarSkeletalMeshLODBias->Set(-10);
+	}
+}
+
+FGeometryCacheLevelSequenceBaker::FConsoleVariableOverrideScope::~FConsoleVariableOverrideScope()
+{
+	if (IConsoleVariable* CVarForceLOD = IConsoleManager::Get().FindConsoleVariable(TEXT("r.ForceLOD")))
+	{
+		CVarForceLOD->Set(PreviousForceLOD);
+	}
+
+	if (IConsoleVariable* CVarSkeletalMeshLODBias = IConsoleManager::Get().FindConsoleVariable(TEXT("r.SkeletalMeshLODBias")))
+	{
+		CVarSkeletalMeshLODBias->Set(PreviousSkeletalMeshLODBias);
+	}
+}
+
+bool FGeometryCacheLevelSequenceBaker::FBakeTask::IsSequencerPlaying()
+{
+	return Sequencer.Pin()->GetPlaybackStatus() == EMovieScenePlayerStatus::Playing;
+}
+
+void FGeometryCacheLevelSequenceBaker::FBakeTask::PlaySequencer()
+{
+	Sequencer.Pin()->SetGlobalTime(StartFrame);
+	Sequencer.Pin()->PlayTo({EndFrame, EUpdatePositionMethod::Play});
+}
+
+void FGeometryCacheLevelSequenceBaker::FBakeTask::SequencerSnapToStart()
+{
+	Sequencer.Pin()->SetGlobalTime(StartFrame);
+}
+
+void FGeometryCacheLevelSequenceBaker::FBakeTask::UpdateGatherProgress()
+{
+	SlowTask->EnterProgressFrame(AmountOfWorkGatherStage/NumSamples);
+}
+
+void FGeometryCacheLevelSequenceBaker::FBakeTask::UpdateBakeProgress()
+{
+	SlowTask->EnterProgressFrame(AmountOfWorkBakeStage/(NumSamples * ComponentTasks.Num()));
+}
+
+void FGeometryCacheLevelSequenceBaker::FBakeTask::TickProgress()
+{
+	SlowTask->TickProgress();
+}
+
+
+#undef LOCTEXT_NAMESPACE
diff --git a/Engine/Plugins/Experimental/GeometryCacheLevelSequenceBaker/Source/GeometryCacheLevelSequenceBaker/Private/GeometryCacheLevelSequenceBaker.h b/Engine/Plugins/Experimental/GeometryCacheLevelSequenceBaker/Source/GeometryCacheLevelSequenceBaker/Private/GeometryCacheLevelSequenceBaker.h
new file mode 100644
--- /dev/null
+++ b/Engine/Plugins/Experimental/GeometryCacheLevelSequenceBaker/Source/GeometryCacheLevelSequenceBaker/Private/GeometryCacheLevelSequenceBaker.h
@@ -0,0 +1,178 @@
+// Copyright Epic Games, Inc. All Rights Reserved.
+
+#pragma once
+
+#include "GeometryCacheConstantTopologyWriter.h"
+#include "ISequencer.h"
+#include "SequencerSettings.h"
+#include "Widgets/SCompoundWidget.h"
+#include "TickableEditorObject.h"
+#include "Misc/ScopedSlowTask.h"
+#include "Engine/SkeletalMesh.h"
+
+#include "GeometryCacheLevelSequenceBaker.generated.h"
+
+class UGeometryCache;
+
+UCLASS()
+class ULevelSequenceGeometryCacheBakerOption : public UObject
+{
+	GENERATED_BODY()
+public:
+	UPROPERTY(EditAnywhere, Category = Options)
+	int32 NumSamplesPerFrame = 1;
+	
+	UPROPERTY()
+	bool bShouldBake = false;
+};
+
+class USkeletalMeshComponent;
+
+class GEOMETRYCACHELEVELSEQUENCEBAKER_API FGeometryCacheLevelSequenceBaker : public FTickableEditorObject
+{
+public:
+
+	static void Bake(TSharedRef<ISequencer> InSequencer );
+	static TArray<FGuid> GetBindingsToBake(TSharedRef<ISequencer> InSequencer);
+	static bool GetGeometryCacheAssetPathFromUser(FString& OutPackageName, FString& OutAssetName);
+
+	void Tick(float DeltaTime) override;
+	void OnSequencerEvaluated();
+	bool IsTickable() const override;
+	TStatId GetStatId() const override {RETURN_QUICK_DECLARE_CYCLE_STAT(FBlueprintActionDatabase, STATGROUP_Tickables);}
+	
+private:
+	FGeometryCacheLevelSequenceBaker() = default;
+	static FGeometryCacheLevelSequenceBaker& Get();
+
+	void SetupComponentBakeTasks();
+	void Gather();
+	void RequestReadback();
+	void WriteToAsset();
+
+	
+
+	static const int32 LODIndexToBake;
+	static const float TotalAmountOfWork;
+	static const float AmountOfWorkGatherStage;
+	static const float AmountOfWorkBakeStage;
+	
+	struct FSkeletalMeshComponentSettingScope
+	{
+		TWeakObjectPtr<USkeletalMeshComponent> Component;
+		
+		bool bPreviousAlwaysUseMeshDeformer = false;
+		bool bPreviousUpdateAnimationInEditor = false;
+		int32 PreviousForcedLOD = 0;
+		
+		FSkeletalMeshComponentSettingScope(USkeletalMeshComponent* InComponent);
+		~FSkeletalMeshComponentSettingScope();
+	};
+
+	using FFrameData = UE::GeometryCacheHelpers::FGeometryCacheConstantTopologyWriter::FFrameData;
+	using FVisibilitySample = UE::GeometryCacheHelpers::FGeometryCacheConstantTopologyWriter::FVisibilitySample;
+
+	struct FComponentInfo
+	{
+		FName Name;
+		USkeletalMesh* SkeletalMeshAsset;
+	};
+	
+	struct FComponentTask
+	{
+		FGuid Binding;
+
+		FComponentInfo ComponentInfo;
+		
+		TArray<FFrameData> GeometrySamples;
+		std::atomic<uint32> NumSamplesPending = 0;
+		TArray<FVisibilitySample> VisibilitySamples;
+
+		TUniquePtr<FSkeletalMeshComponentSettingScope> ComponentSettingScope;
+	};
+
+	static void OnReadbackResultConfirmed(FComponentTask& Task, int32 SampleIndex, bool bMeshAvailable);
+
+	struct FEngineFixedDeltaTimeScope
+	{
+		bool bPreviousUseFixedDeltaTime = false;
+		double PreviousFixedDeltaTime = 0;
+		FEngineFixedDeltaTimeScope(double NewFixedDeltaTime);
+		~FEngineFixedDeltaTimeScope();
+	};
+
+	struct FSequencerSettingScope
+	{
+		TWeakPtr<ISequencer> Sequencer;
+		FQualifiedFrameTime PreviousGlobalTime;
+		ESequencerLoopMode PreviousLoopMode;
+		FDelegateHandle OnGlobalTimeChangedDelegateHandle;
+
+		
+		FSequencerSettingScope(TSharedPtr<ISequencer> InSequencer);
+		~FSequencerSettingScope();
+	};
+
+	struct FConsoleVariableOverrideScope
+	{
+		int32 PreviousForceLOD = 0;
+		int32 PreviousSkeletalMeshLODBias = 0;
+		FConsoleVariableOverrideScope();
+		~FConsoleVariableOverrideScope();
+	};
+
+	enum class EStage
+	{
+		Gather,
+		RequestReadback,
+		WriteToAsset,
+		End
+	};
+
+
+	
+	struct FBakeTask
+	{
+		FString PackageName;
+		FString AssetName;
+
+		TWeakPtr<ISequencer> Sequencer;
+		TArray<FGuid> Bindings;
+		TMap<FGuid, TMap<FName, FComponentInfo>> BindingToComponentInfoMap;
+		FFrameTime StartFrame;
+		FFrameTime EndFrame;
+		
+		EStage Stage = EStage::Gather;
+		
+		TArray<FComponentTask> ComponentTasks;
+		std::atomic<uint32> NumComponentTasksPending = 0;
+		
+		uint32 NumSamples = 0;
+		uint32 CurrentSampleIndex = 0;
+		float SamplesPerSecond = 0;
+
+		TUniquePtr<FScopedSlowTask> SlowTask;
+		TUniquePtr<FEngineFixedDeltaTimeScope> FixedDeltaTimeScope;
+		TUniquePtr<FSequencerSettingScope> SequencerStateScope;
+		TUniquePtr<FConsoleVariableOverrideScope> ConsoleVariableOverrideScope;
+
+		FDelegateHandle OnEndFrameDelegateHandle;
+
+		TMap<TWeakObjectPtr<USkeletalMeshComponent>, TUniquePtr<FSkeletalMeshComponentSettingScope>> GatherStageComponentSettingScopes;
+
+		bool IsSequencerPlaying();
+		void PlaySequencer();
+		void SequencerSnapToStart();
+		
+		void UpdateGatherProgress();
+		void UpdateBakeProgress();
+		void TickProgress();
+	};
+
+	FCriticalSection CurrentBakeTaskLifeTimeCriticalSection;
+	TUniquePtr<FBakeTask> CurrentBakeTask;
+};
+
+
+
+
diff --git a/Engine/Plugins/Experimental/GeometryCacheLevelSequenceBaker/Source/GeometryCacheLevelSequenceBaker/Private/GeometryCacheLevelSequenceBakerCommands.cpp b/Engine/Plugins/Experimental/GeometryCacheLevelSequenceBaker/Source/GeometryCacheLevelSequenceBaker/Private/GeometryCacheLevelSequenceBakerCommands.cpp
new file mode 100644
--- /dev/null
+++ b/Engine/Plugins/Experimental/GeometryCacheLevelSequenceBaker/Source/GeometryCacheLevelSequenceBaker/Private/GeometryCacheLevelSequenceBakerCommands.cpp
@@ -0,0 +1,24 @@
+// Copyright Epic Games, Inc. All Rights Reserved.
+
+#include "GeometryCacheLevelSequenceBakerCommands.h"
+#include "GeometryCacheLevelSequenceBakerStyle.h"
+
+#define LOCTEXT_NAMESPACE "FGeometryCacheLevelSequenceBakerCommands"
+
+FGeometryCacheLevelSequenceBakerCommands::FGeometryCacheLevelSequenceBakerCommands()
+	: TCommands<FGeometryCacheLevelSequenceBakerCommands>(
+		"GeometryCacheLevelSequenceBakerCommands", // Context name for fast lookup
+		NSLOCTEXT("Contexts", "GeometryCacheLevelSequenceBakerCommands", "Geometry Cache Level Sequence Baker"), // Localized context name for displaying
+		NAME_None, // Parent
+		FGeometryCacheLevelSequenceBakerStyle::Get().GetStyleSetName() // Icon Style Set
+		)
+{
+}
+
+void FGeometryCacheLevelSequenceBakerCommands::RegisterCommands()
+{
+	UI_COMMAND(BakeGeometryCache, "Bake Geometry Cache", "Bake the selected skeletal meshes (all if none selected) to a Geometry Cache. (Shots and sub-scenes not supported)", EUserInterfaceActionType::Button, FInputChord());
+}
+
+#undef LOCTEXT_NAMESPACE
+
diff --git a/Engine/Plugins/Experimental/GeometryCacheLevelSequenceBaker/Source/GeometryCacheLevelSequenceBaker/Private/GeometryCacheLevelSequenceBakerCommands.h b/Engine/Plugins/Experimental/GeometryCacheLevelSequenceBaker/Source/GeometryCacheLevelSequenceBaker/Private/GeometryCacheLevelSequenceBakerCommands.h
new file mode 100644
--- /dev/null
+++ b/Engine/Plugins/Experimental/GeometryCacheLevelSequenceBaker/Source/GeometryCacheLevelSequenceBaker/Private/GeometryCacheLevelSequenceBakerCommands.h
@@ -0,0 +1,17 @@
+// Copyright Epic Games, Inc. All Rights Reserved.
+
+#pragma once
+
+#include "CoreMinimal.h"
+#include "Framework/Commands/Commands.h"
+
+class FGeometryCacheLevelSequenceBakerCommands : public TCommands<FGeometryCacheLevelSequenceBakerCommands>
+{
+public:
+	FGeometryCacheLevelSequenceBakerCommands();
+
+	/** Exports sequence to geometry cache. */
+	TSharedPtr< FUICommandInfo > BakeGeometryCache;
+	
+	virtual void RegisterCommands() override;
+};
diff --git a/Engine/Plugins/Experimental/GeometryCacheLevelSequenceBaker/Source/GeometryCacheLevelSequenceBaker/Private/GeometryCacheLevelSequenceBakerModule.cpp b/Engine/Plugins/Experimental/GeometryCacheLevelSequenceBaker/Source/GeometryCacheLevelSequenceBaker/Private/GeometryCacheLevelSequenceBakerModule.cpp
new file mode 100644
--- /dev/null
+++ b/Engine/Plugins/Experimental/GeometryCacheLevelSequenceBaker/Source/GeometryCacheLevelSequenceBaker/Private/GeometryCacheLevelSequenceBakerModule.cpp
@@ -0,0 +1,39 @@
+﻿// Copyright Epic Games, Inc. All Rights Reserved.
+
+#include "GeometryCacheLevelSequenceBakerModule.h"
+
+#include "FGeometryCacheLevelSequenceBakerCustomization.h"
+#include "GeometryCacheLevelSequenceBakerCommands.h"
+#include "GeometryCacheLevelSequenceBakerStyle.h"
+#include "LevelSequenceEditorModule.h"
+#include "SequencerCustomizationManager.h"
+
+DEFINE_LOG_CATEGORY(LogGeometryCacheLevelSequenceBaker);
+
+#define LOCTEXT_NAMESPACE "FGeometryCacheLevelSequenceBakerModule"
+
+void FGeometryCacheLevelSequenceBakerModule::StartupModule()
+{
+	FGeometryCacheLevelSequenceBakerStyle::Register();	
+	FGeometryCacheLevelSequenceBakerCommands::Register();
+	
+	ILevelSequenceEditorModule& LevelSequenceEditorModule = FModuleManager::Get().LoadModuleChecked<ILevelSequenceEditorModule>("LevelSequenceEditor");
+	CustomizationHandle = LevelSequenceEditorModule.RegisterAdditionalLevelSequenceEditorCustomization(FOnGetSequencerCustomizationInstance::CreateLambda([]()
+	{
+		return new FGeometryCacheLevelSequenceBakerCustomization();
+	}));
+}
+
+void FGeometryCacheLevelSequenceBakerModule::ShutdownModule()
+{
+	ILevelSequenceEditorModule& LevelSequenceEditorModule = FModuleManager::Get().LoadModuleChecked<ILevelSequenceEditorModule>("LevelSequenceEditor");
+	LevelSequenceEditorModule.UnregisterAdditionalLevelSequenceEditorCustomization(CustomizationHandle);
+}
+
+
+
+
+
+#undef LOCTEXT_NAMESPACE
+	
+IMPLEMENT_MODULE(FGeometryCacheLevelSequenceBakerModule, GeometryCacheLevelSequenceBaker)
\ No newline at end of file
diff --git a/Engine/Plugins/Experimental/GeometryCacheLevelSequenceBaker/Source/GeometryCacheLevelSequenceBaker/Private/GeometryCacheLevelSequenceBakerStyle.cpp b/Engine/Plugins/Experimental/GeometryCacheLevelSequenceBaker/Source/GeometryCacheLevelSequenceBaker/Private/GeometryCacheLevelSequenceBakerStyle.cpp
new file mode 100644
--- /dev/null
+++ b/Engine/Plugins/Experimental/GeometryCacheLevelSequenceBaker/Source/GeometryCacheLevelSequenceBaker/Private/GeometryCacheLevelSequenceBakerStyle.cpp
@@ -0,0 +1,40 @@
+// Copyright Epic Games, Inc. All Rights Reserved.
+
+#include "GeometryCacheLevelSequenceBakerStyle.h"
+
+#include "Styling/SlateStyleMacros.h"
+#include "Styling/SlateStyleRegistry.h"
+
+
+FGeometryCacheLevelSequenceBakerStyle& FGeometryCacheLevelSequenceBakerStyle::Get()
+{
+	static FGeometryCacheLevelSequenceBakerStyle Instance;
+	return Instance;
+}
+
+void FGeometryCacheLevelSequenceBakerStyle::Register()
+{
+	FSlateStyleRegistry::RegisterSlateStyle(Get());
+}
+
+void FGeometryCacheLevelSequenceBakerStyle::Unregister()
+{
+	FSlateStyleRegistry::UnRegisterSlateStyle(Get());
+}
+
+FGeometryCacheLevelSequenceBakerStyle::FGeometryCacheLevelSequenceBakerStyle()
+	: FSlateStyleSet("GeometryCacheLevelSequenceBakerStyle")
+{
+	static const FVector2D IconSize10x10(10.0f, 10.0f);
+	static const FVector2D IconSize16x12(16.0f, 12.0f);
+	static const FVector2D IconSize16x16(16.0f, 16.0f);
+	static const FVector2D IconSize20x20(20.0f, 20.0f);
+	static const FVector2D IconSize64x64(64.0f, 64.0f);
+
+	static const FSlateColor DefaultForeground(FLinearColor(0.72f, 0.72f, 0.72f, 1.f));
+
+	FSlateStyleSet::SetContentRoot(FPaths::EnginePluginsDir() / TEXT("Experimental/GeometryCacheLevelSequenceBaker/Resources"));
+	{
+		Set("GeometryCacheLevelSequenceBakerCommands.BakeGeometryCache", new IMAGE_BRUSH("Icon128", IconSize20x20));
+	}
+}
diff --git a/Engine/Plugins/Experimental/GeometryCacheLevelSequenceBaker/Source/GeometryCacheLevelSequenceBaker/Private/GeometryCacheLevelSequenceBakerStyle.h b/Engine/Plugins/Experimental/GeometryCacheLevelSequenceBaker/Source/GeometryCacheLevelSequenceBaker/Private/GeometryCacheLevelSequenceBakerStyle.h
new file mode 100644
--- /dev/null
+++ b/Engine/Plugins/Experimental/GeometryCacheLevelSequenceBaker/Source/GeometryCacheLevelSequenceBaker/Private/GeometryCacheLevelSequenceBakerStyle.h
@@ -0,0 +1,22 @@
+// Copyright Epic Games, Inc. All Rights Reserved.
+
+#pragma once
+
+#include "Styling/SlateStyle.h"
+
+
+class FGeometryCacheLevelSequenceBakerStyle
+	: public FSlateStyleSet
+{
+public:
+	static FGeometryCacheLevelSequenceBakerStyle& Get();
+
+protected:
+	friend class FGeometryCacheLevelSequenceBakerModule;
+
+	static void Register();
+	static void Unregister();
+
+private:
+	FGeometryCacheLevelSequenceBakerStyle();
+};
\ No newline at end of file
diff --git a/Engine/Plugins/Experimental/GeometryCacheLevelSequenceBaker/Source/GeometryCacheLevelSequenceBaker/Public/GeometryCacheLevelSequenceBakerModule.h b/Engine/Plugins/Experimental/GeometryCacheLevelSequenceBaker/Source/GeometryCacheLevelSequenceBaker/Public/GeometryCacheLevelSequenceBakerModule.h
new file mode 100644
--- /dev/null
+++ b/Engine/Plugins/Experimental/GeometryCacheLevelSequenceBaker/Source/GeometryCacheLevelSequenceBaker/Public/GeometryCacheLevelSequenceBakerModule.h
@@ -0,0 +1,18 @@
+﻿// Copyright Epic Games, Inc. All Rights Reserved.
+
+#pragma once
+
+#include "Modules/ModuleManager.h"
+#include "Misc/Guid.h"
+
+DECLARE_LOG_CATEGORY_EXTERN(LogGeometryCacheLevelSequenceBaker, Log, All);
+
+class FGeometryCacheLevelSequenceBakerModule : public IModuleInterface
+{
+public:
+	/** IModuleInterface implementation */
+	virtual void StartupModule() override;
+	virtual void ShutdownModule() override;
+
+	FGuid CustomizationHandle;
+};
diff --git a/Engine/Plugins/MovieScene/LevelSequenceEditor/Source/LevelSequenceEditor/Private/LevelSequenceEditorModule.cpp b/Engine/Plugins/MovieScene/LevelSequenceEditor/Source/LevelSequenceEditor/Private/LevelSequenceEditorModule.cpp
--- a/Engine/Plugins/MovieScene/LevelSequenceEditor/Source/LevelSequenceEditor/Private/LevelSequenceEditorModule.cpp
+++ b/Engine/Plugins/MovieScene/LevelSequenceEditor/Source/LevelSequenceEditor/Private/LevelSequenceEditorModule.cpp
@@ -201,9 +201,16 @@ protected:
 	{
 		ISequencerModule& SequencerModule = FModuleManager::LoadModuleChecked<ISequencerModule>("Sequencer");
 		SequencerModule.GetSequencerCustomizationManager()->RegisterInstancedSequencerCustomization(ULevelSequence::StaticClass(),
-				FOnGetSequencerCustomizationInstance::CreateLambda([]()
+				FOnGetSequencerCustomizationInstance::CreateLambda([this]()
 				{
-					return new UE::Sequencer::FLevelSequenceCustomization();
+					UE::Sequencer::FLevelSequenceCustomization* Customization = new UE::Sequencer::FLevelSequenceCustomization();
+					TArray<FOnGetSequencerCustomizationInstance> CustomizationInstances;
+					AdditionalLevelSequenceEditorCustomizations.GenerateValueArray(CustomizationInstances);
+					for (const FOnGetSequencerCustomizationInstance& CustomizationInstance : CustomizationInstances)
+					{
+						Customization->AddCustomization(TUniquePtr<ISequencerCustomization>(CustomizationInstance.Execute()));
+					}
+					return Customization;
 				}));
 	}
 
@@ -362,6 +369,19 @@ protected:
 		return OnComputePlaybackContextDelegate;
 	}
 
+	FGuid RegisterAdditionalLevelSequenceEditorCustomization(FOnGetSequencerCustomizationInstance OnGetInstance) override
+	{
+		FGuid Handle = FGuid::NewGuid();
+		AdditionalLevelSequenceEditorCustomizations.Add(Handle, OnGetInstance);
+		
+		return Handle;
+	}
+	
+	void UnregisterAdditionalLevelSequenceEditorCustomization(FGuid InHandle) override
+	{
+		AdditionalLevelSequenceEditorCustomizations.Remove(InHandle);
+	}
+
 	static TSharedRef<ISequencerEditorObjectBinding> OnCreateActorBinding(TSharedRef<ISequencer> InSequencer)
 	{
 		return MakeShareable(new FLevelSequenceEditorActorBinding(InSequencer));
@@ -399,6 +419,8 @@ private:
 	TObjectPtr<USequencerSettings> Settings;
 
 	FDelegateHandle SequenceEditorHandle;
+
+	TMap<FGuid, FOnGetSequencerCustomizationInstance> AdditionalLevelSequenceEditorCustomizations;
 };
 
 
diff --git a/Engine/Plugins/MovieScene/LevelSequenceEditor/Source/LevelSequenceEditor/Private/Misc/LevelSequenceCustomization.cpp b/Engine/Plugins/MovieScene/LevelSequenceEditor/Source/LevelSequenceEditor/Private/Misc/LevelSequenceCustomization.cpp
--- a/Engine/Plugins/MovieScene/LevelSequenceEditor/Source/LevelSequenceEditor/Private/Misc/LevelSequenceCustomization.cpp
+++ b/Engine/Plugins/MovieScene/LevelSequenceEditor/Source/LevelSequenceEditor/Private/Misc/LevelSequenceCustomization.cpp
@@ -30,6 +30,11 @@
 namespace UE::Sequencer
 {
 
+void FLevelSequenceCustomization::AddCustomization(TUniquePtr<ISequencerCustomization> NewCustomization)
+{
+	AdditionalCustomizations.Add(MoveTemp(NewCustomization));
+}
+
 void FLevelSequenceCustomization::RegisterSequencerCustomization(FSequencerCustomizationBuilder& Builder)
 {
 	WeakSequencer = Builder.GetSequencer().AsShared();
@@ -60,10 +65,21 @@ void FLevelSequenceCustomization::RegisterSequencerCustomization(FSequencerCusto
 	Customization.OnBuildObjectBindingContextMenu = FOnGetSequencerMenuExtender::CreateRaw(this, &FLevelSequenceCustomization::CreateObjectBindingContextMenuExtender);
 	Customization.OnBuildSidebarMenu = FOnGetSequencerMenuExtender::CreateRaw(this, &FLevelSequenceCustomization::CreateObjectBindingSidebarMenuExtender);
 	Builder.AddCustomization(Customization);
+
+	for(TUniquePtr<ISequencerCustomization>& ExternalCustomization : AdditionalCustomizations)
+	{
+		ExternalCustomization->RegisterSequencerCustomization(Builder);
+	}
+	
 }
 
 void FLevelSequenceCustomization::UnregisterSequencerCustomization()
 {
+	for(TUniquePtr<ISequencerCustomization>& ExternalCustomization : AdditionalCustomizations)
+	{
+		ExternalCustomization->UnregisterSequencerCustomization();
+	}
+
 	ISequencerModule& SequencerModule = FModuleManager::LoadModuleChecked<ISequencerModule>("Sequencer");
 	SequencerModule.GetActionsMenuExtensibilityManager()->RemoveExtender(ActionsMenuExtender);
 
diff --git a/Engine/Plugins/MovieScene/LevelSequenceEditor/Source/LevelSequenceEditor/Private/Misc/LevelSequenceCustomization.h b/Engine/Plugins/MovieScene/LevelSequenceEditor/Source/LevelSequenceEditor/Private/Misc/LevelSequenceCustomization.h
--- a/Engine/Plugins/MovieScene/LevelSequenceEditor/Source/LevelSequenceEditor/Private/Misc/LevelSequenceCustomization.h
+++ b/Engine/Plugins/MovieScene/LevelSequenceEditor/Source/LevelSequenceEditor/Private/Misc/LevelSequenceCustomization.h
@@ -17,6 +17,9 @@ namespace UE::Sequencer
  */
 class FLevelSequenceCustomization : public ISequencerCustomization
 {
+public:
+	void AddCustomization(TUniquePtr<ISequencerCustomization> NewCustomization);
+	
 protected:
 
 	//~ ISequencerCustomization interface
@@ -45,6 +48,8 @@ private:
 
 	TSharedPtr<FUICommandList> ActionsMenuCommandList;
 	TSharedPtr<FExtender> ActionsMenuExtender;
+
+	TArray<TUniquePtr<ISequencerCustomization>>	AdditionalCustomizations;
 };
 
 } // namespace UE::Sequencer
diff --git a/Engine/Plugins/MovieScene/LevelSequenceEditor/Source/LevelSequenceEditor/Public/LevelSequenceEditorModule.h b/Engine/Plugins/MovieScene/LevelSequenceEditor/Source/LevelSequenceEditor/Public/LevelSequenceEditorModule.h
--- a/Engine/Plugins/MovieScene/LevelSequenceEditor/Source/LevelSequenceEditor/Public/LevelSequenceEditorModule.h
+++ b/Engine/Plugins/MovieScene/LevelSequenceEditor/Source/LevelSequenceEditor/Public/LevelSequenceEditorModule.h
@@ -2,6 +2,7 @@
 
 #pragma once
 
+#include "SequencerCustomizationManager.h"
 #include "Delegates/Delegate.h"
 #include "Modules/ModuleInterface.h"
 
@@ -21,4 +22,7 @@ public:
 
 	DECLARE_MULTICAST_DELEGATE_OneParam(FAllowPlaybackContext, bool&);
 	virtual FAllowPlaybackContext& OnComputePlaybackContext() = 0;
+
+	virtual FGuid RegisterAdditionalLevelSequenceEditorCustomization(FOnGetSequencerCustomizationInstance OnGetInstance) = 0;
+	virtual void UnregisterAdditionalLevelSequenceEditorCustomization(FGuid InHandle) = 0;
 };
diff --git a/Engine/Plugins/Runtime/GeometryCache/Source/GeometryCache/Private/GeometryCacheConstantTopologyWriter.cpp b/Engine/Plugins/Runtime/GeometryCache/Source/GeometryCache/Private/GeometryCacheConstantTopologyWriter.cpp
--- a/Engine/Plugins/Runtime/GeometryCache/Source/GeometryCache/Private/GeometryCacheConstantTopologyWriter.cpp
+++ b/Engine/Plugins/Runtime/GeometryCache/Source/GeometryCache/Private/GeometryCacheConstantTopologyWriter.cpp
@@ -297,6 +297,11 @@ namespace UE::GeometryCacheHelpers
 	}
 
 	bool FGeometryCacheConstantTopologyWriter::FTrackWriter::WriteAndClose(TArrayView<FFrameData> FramesToMoveFrom)
+	{
+		return WriteAndClose(FramesToMoveFrom, {});
+	}
+
+	bool FTrackWriter::WriteAndClose(TArrayView<FFrameData> FramesToMoveFrom, const TArray<FVisibilitySample>& VisibilitySamples)
 	{
 		if (!Owner->Cache.IsValid() || !Track.IsValid())
 		{
@@ -424,6 +429,12 @@ namespace UE::GeometryCacheHelpers
 		TArray<FMatrix> Mats { FMatrix::Identity, FMatrix::Identity };
 		TArray<float> MatTimes { 0.0f, static_cast<float>(MaxRecordedFrame) / ConfigRef.FPS };
 		TrackStreamable->SetMatrixSamples(Mats, MatTimes);
+
+		for (const FVisibilitySample& Sample : VisibilitySamples)
+		{
+			const float Time = static_cast<float>(Sample.FrameIndex) / ConfigRef.FPS;
+			TrackStreamable->AddVisibilitySample(Sample.bVisible, Time);
+		}
 		
 		if (ensureAlways(TrackStreamable->EndCoding()))
 		{
diff --git a/Engine/Plugins/Runtime/GeometryCache/Source/GeometryCache/Private/GeometryCacheSceneProxy.cpp b/Engine/Plugins/Runtime/GeometryCache/Source/GeometryCache/Private/GeometryCacheSceneProxy.cpp
--- a/Engine/Plugins/Runtime/GeometryCache/Source/GeometryCache/Private/GeometryCacheSceneProxy.cpp
+++ b/Engine/Plugins/Runtime/GeometryCache/Source/GeometryCache/Private/GeometryCacheSceneProxy.cpp
@@ -71,7 +71,7 @@ static TAutoConsoleVariable<int32> CVarInterpolateFrames(
 
 static TAutoConsoleVariable<int32> CVarRayTracingGeometryCache(
 	TEXT("r.RayTracing.Geometry.GeometryCache"),
-	1,
+	0,
 	TEXT("Include geometry cache primitives in ray tracing effects (default = 1 (geometry cache enabled in ray tracing))"));
 
 /**
@@ -198,6 +198,12 @@ FGeometryCacheSceneProxy::FGeometryCacheSceneProxy(UGeometryCacheComponent* Comp
 				{
 					if (Section != nullptr)
 					{
+						const FVisibilitySample& VisibilitySample = Section->GetVisibilitySample(SceneProxy->Time, SceneProxy->bLooping);
+						if (!VisibilitySample.bVisibilityState)
+						{
+							continue;
+						}
+			
 						FRayTracingGeometryInitializer Initializer;
 						Initializer.DebugName = DebugName;
 						const int PositionBufferIndex = Section->CurrentPositionBufferIndex != -1 ? Section->CurrentPositionBufferIndex % 2 : 0;
@@ -621,6 +627,12 @@ void FGeometryCacheSceneProxy::UpdateAnimation(FRHICommandListBase& RHICmdList,
 		{
 			if (Section != nullptr)
 			{
+				const FVisibilitySample& VisibilitySample = Section->GetVisibilitySample(Time, bLooping);
+				if (!VisibilitySample.bVisibilityState)
+				{
+					continue;
+				}
+			
 				const int PositionBufferIndex = Section->CurrentPositionBufferIndex != -1 ? Section->CurrentPositionBufferIndex % 2 : 0;
 				const FGeometryCacheMeshData* MeshData = Section->bNextFrameMeshDataSelected ? Section->NextFrameMeshData : Section->MeshData;
 				const uint32 IndexBufferNumTriangles = Section->IndexBuffer.NumValidIndices / 3;
diff --git a/Engine/Plugins/Runtime/GeometryCache/Source/GeometryCache/Private/GeometryCacheTrackStreamable.cpp b/Engine/Plugins/Runtime/GeometryCache/Source/GeometryCache/Private/GeometryCacheTrackStreamable.cpp
--- a/Engine/Plugins/Runtime/GeometryCache/Source/GeometryCache/Private/GeometryCacheTrackStreamable.cpp
+++ b/Engine/Plugins/Runtime/GeometryCache/Source/GeometryCache/Private/GeometryCacheTrackStreamable.cpp
@@ -229,8 +229,20 @@ bool UGeometryCacheTrackStreamable::EndCoding()
 				FVisibilitySample Sample(ImportVisibilitySamples[SampleIndex].Value);
 				Sample.Range = VisibilityRange;
 				VisibilitySamples.Add(Sample);
+
+				bVisible = ImportVisibilitySamples[SampleIndex].Value;
+				RangeStart = VisibilityRange.GetUpperBoundValue();	
 			}
 		}
+
+		// Make sure the range of all visibility samples covers the entire track duration
+		if (RangeStart < Samples.Last().SampleTime)
+		{
+			TRange<float> VisibilityRange(RangeStart, Samples.Last().SampleTime);
+			FVisibilitySample Sample(bVisible);
+			Sample.Range = VisibilityRange;
+			VisibilitySamples.Add(Sample);	
+		}
 	}
 	else
 	{
diff --git a/Engine/Plugins/Runtime/GeometryCache/Source/GeometryCache/Public/GeometryCacheConstantTopologyWriter.h b/Engine/Plugins/Runtime/GeometryCache/Source/GeometryCache/Public/GeometryCacheConstantTopologyWriter.h
--- a/Engine/Plugins/Runtime/GeometryCache/Source/GeometryCache/Public/GeometryCacheConstantTopologyWriter.h
+++ b/Engine/Plugins/Runtime/GeometryCache/Source/GeometryCache/Public/GeometryCacheConstantTopologyWriter.h
@@ -58,6 +58,12 @@ namespace UE::GeometryCacheHelpers
 			TArray<FVector3f> TangentsX;
 		};
 
+		struct FVisibilitySample
+		{
+			int32 FrameIndex = 0;
+			bool bVisible = true;
+		};
+
 		struct FTrackWriter
 		{
 			GEOMETRYCACHE_API FTrackWriter(FGeometryCacheConstantTopologyWriter& InOwner, FName TrackName = FName());
@@ -91,6 +97,14 @@ namespace UE::GeometryCacheHelpers
 			 * Normals and tangents must have the same size as positions.
 			 */
 			GEOMETRYCACHE_API bool WriteAndClose(TArrayView<FFrameData> FramesToMoveFrom);
+
+			/**
+			 * Move the frame data to the cache track and close the TrackWriter.
+			 * Similar to WriteAndClose(TArrayView<FFrameData> FramesToMoveFrom), but also supports visibility samples.
+			 * Each visibility sample defines the visibility of the track starting from the specified frame to the next visibility sample
+			 */
+			GEOMETRYCACHE_API bool WriteAndClose(TArrayView<FFrameData> FramesToMoveFrom, const TArray<FVisibilitySample>& VisibilitySamples);
+			
 		private:
 			TStrongObjectPtr<UGeometryCacheTrack> Track;
 			FGeometryCacheConstantTopologyWriter* Owner = nullptr;
diff --git a/Engine/Source/Runtime/Engine/Classes/Animation/MeshDeformerInstance.h b/Engine/Source/Runtime/Engine/Classes/Animation/MeshDeformerInstance.h
--- a/Engine/Source/Runtime/Engine/Classes/Animation/MeshDeformerInstance.h
+++ b/Engine/Source/Runtime/Engine/Classes/Animation/MeshDeformerInstance.h
@@ -3,6 +3,11 @@
 #pragma once
 
 #include "CoreMinimal.h"
+
+#if WITH_EDITORONLY_DATA
+#include "Animation/MeshDeformerGeometryReadback.h"
+#endif // WITH_EDITORONLY_DATA
+
 #include "MeshDeformerInstance.generated.h"
 
 class FSceneInterface;
@@ -86,7 +91,7 @@ public:
 
 #if WITH_EDITORONLY_DATA
 	/** Reads back the deformed geometry and generates a mesh description */
-	ENGINE_API virtual bool RequestReadbackDeformerGeometry(TFunction<void(const FMeshDescription&)> InOnReadbackCompleted_AnyThread) PURE_VIRTUAL(, return false; );
+	ENGINE_API virtual bool RequestReadbackDeformerGeometry(TUniquePtr<FMeshDeformerGeometryReadbackRequest> InRequest) PURE_VIRTUAL(, return false; );
 #endif // WITH_EDITORONLY_DATA
 	
 	/** Returns the specific instance that directly represents the source deformer, this is needed as a deformer may create intermediate instances that aren't
diff --git a/Engine/Source/Runtime/Engine/Classes/Components/SkinnedMeshComponent.h b/Engine/Source/Runtime/Engine/Classes/Components/SkinnedMeshComponent.h
--- a/Engine/Source/Runtime/Engine/Classes/Components/SkinnedMeshComponent.h
+++ b/Engine/Source/Runtime/Engine/Classes/Components/SkinnedMeshComponent.h
@@ -17,6 +17,11 @@
 #include "BoneContainer.h"
 #include "ClothingSystemRuntimeTypes.h"
 #include "StateStream/SkinnedMeshStateStreamHandle.h"
+
+#if WITH_EDITORONLY_DATA
+#include "Animation/MeshDeformerGeometryReadback.h"
+#endif // WITH_EDITORONLY_DATA
+
 #include "SkinnedMeshComponent.generated.h"
 
 enum class ESkinCacheUsage : uint8;
@@ -348,13 +353,13 @@ public:
 
 #if WITH_EDITORONLY_DATA
 private:
-	TArray<TFunction<void(const FMeshDescription&)>> RenderGeometryReadbackCallbacks;
+	TArray<TUniquePtr<FMeshDeformerGeometryReadbackRequest>> GeometryReadbackRequests;
 public:
 	/**
 	 * Readback the render geometry async, a FMeshDescription will be generated on a background thread for additional operation
 	 * Note: currently it only supports mesh deformed by a MeshDeformer
 	 */
-	ENGINE_API bool RequestReadbackRenderGeometry(TFunction<void(const FMeshDescription&)> InOnReadbackCompleted_AnyThread);
+	ENGINE_API bool RequestReadbackRenderGeometry(TUniquePtr<FMeshDeformerGeometryReadbackRequest> InRequest);
 #endif // WITH_EDITORONLY_DATA
 	
 public:
diff --git a/Engine/Source/Runtime/Engine/Private/Animation/MeshDeformerGeometryReadback.cpp b/Engine/Source/Runtime/Engine/Private/Animation/MeshDeformerGeometryReadback.cpp
new file mode 100644
--- /dev/null
+++ b/Engine/Source/Runtime/Engine/Private/Animation/MeshDeformerGeometryReadback.cpp
@@ -0,0 +1,16 @@
+﻿// Copyright Epic Games, Inc. All Rights Reserved.
+
+#include "Animation/MeshDeformerGeometryReadback.h"
+#include "MeshDescription.h"
+
+FMeshDeformerGeometryReadbackRequest::~FMeshDeformerGeometryReadbackRequest()
+{
+	if (!bHandled)
+	{
+		if (MeshDescriptionCallback_AnyThread.IsSet())
+		{
+			// Signal to the requester that the request was not successfully fullfilled
+			MeshDescriptionCallback_AnyThread(FMeshDescription());
+		}
+	}
+}
\ No newline at end of file
diff --git a/Engine/Source/Runtime/Engine/Private/Components/SkinnedMeshComponent.cpp b/Engine/Source/Runtime/Engine/Private/Components/SkinnedMeshComponent.cpp
--- a/Engine/Source/Runtime/Engine/Private/Components/SkinnedMeshComponent.cpp
+++ b/Engine/Source/Runtime/Engine/Private/Components/SkinnedMeshComponent.cpp
@@ -50,6 +50,7 @@
 #include "ComponentRecreateRenderStateContext.h"
 #include "NaniteSceneProxy.h"
 #include "NaniteVertexFactory.h"
+#include "Animation/MeshDeformerGeometryReadback.h"
 
 DEFINE_LOG_CATEGORY_STATIC(LogSkinnedMeshComp, Log, All);
 
@@ -1175,6 +1176,14 @@ void USkinnedMeshComponent::CreateRenderState_Concurrent(FRegisterComponentConte
 
 	if (UMeshDeformerInstance* MeshDeformerInstance = GetMeshDeformerInstanceForLOD(UseLOD))
  	{
+
+#if WITH_EDITORONLY_DATA
+			for (TUniquePtr<FMeshDeformerGeometryReadbackRequest>& Request : GeometryReadbackRequests)
+			{
+				MeshDeformerInstance->RequestReadbackDeformerGeometry(MoveTemp(Request));
+			}
+#endif // WITH_EDITORONLY_DATA
+
 		// Enqueue immediate execution of work here to ensure that we have some deformer outputs written for the next frame.
 		UMeshDeformerInstance::FEnqueueWorkDesc Desc;
 		Desc.Scene = GetScene();
@@ -1182,6 +1191,12 @@ void USkinnedMeshComponent::CreateRenderState_Concurrent(FRegisterComponentConte
 		Desc.OwnerName = GetSkinnedAsset() != nullptr ? GetSkinnedAsset()->GetFName() : GetFName();
 		MeshDeformerInstance->EnqueueWork(Desc);
 	}
+
+#if WITH_EDITORONLY_DATA
+	// Avoid unbound accumulation of readback requests, each readback request must be picked up by the deformer instance within the same frame or ignored otherwise
+	GeometryReadbackRequests.Reset();
+#endif // WITH_EDITORONLY_DATA 
+	
 }
 
 void USkinnedMeshComponent::DestroyRenderState_Concurrent()
@@ -1306,9 +1321,9 @@ void USkinnedMeshComponent::SendRenderDynamicData_Concurrent()
 		if (DeformerInstanceForLOD)
 		{
 #if WITH_EDITORONLY_DATA
-			for (TFunction<void(const FMeshDescription&)> ReadbackCallback : RenderGeometryReadbackCallbacks)
+			for (TUniquePtr<FMeshDeformerGeometryReadbackRequest>& Request : GeometryReadbackRequests)
 			{
-				DeformerInstanceForLOD->RequestReadbackDeformerGeometry(ReadbackCallback);
+				DeformerInstanceForLOD->RequestReadbackDeformerGeometry(MoveTemp(Request));
 			}
 #endif // WITH_EDITORONLY_DATA
 			
@@ -1326,7 +1341,7 @@ void USkinnedMeshComponent::SendRenderDynamicData_Concurrent()
 
 #if WITH_EDITORONLY_DATA
 	// Avoid unbound accumulation of readback requests, each readback request must be picked up by the deformer instance within the same frame or ignored otherwise
-	RenderGeometryReadbackCallbacks.Reset();
+	GeometryReadbackRequests.Reset();
 #endif // WITH_EDITORONLY_DATA 
 }
 
@@ -2625,7 +2640,7 @@ int32 USkinnedMeshComponent::GetMeshDeformerMaxLOD() const
 }
 
 #if WITH_EDITORONLY_DATA
-bool USkinnedMeshComponent::RequestReadbackRenderGeometry(TFunction<void(const FMeshDescription&)> InOnReadbackCompleted_AnyThread)
+bool USkinnedMeshComponent::RequestReadbackRenderGeometry(TUniquePtr<FMeshDeformerGeometryReadbackRequest> InRequest)
 {
 	if (!IsInGameThread())
 	{
@@ -2638,8 +2653,8 @@ bool USkinnedMeshComponent::RequestReadbackRenderGeometry(TFunction<void(const F
 		return false;
 	}
 	
-	// These readback requests will be picked up by the deformer instance during SendRenderDynamicData
-	RenderGeometryReadbackCallbacks.Add(InOnReadbackCompleted_AnyThread);
+	// These readback requests will be picked up by the deformer instance during CreateRenderStage/SendRenderDynamicData
+	GeometryReadbackRequests.Add(MoveTemp(InRequest));
 	
 	return true;
 }
diff --git a/Engine/Source/Runtime/Engine/Public/Animation/MeshDeformerGeometryReadback.h b/Engine/Source/Runtime/Engine/Public/Animation/MeshDeformerGeometryReadback.h
new file mode 100644
--- /dev/null
+++ b/Engine/Source/Runtime/Engine/Public/Animation/MeshDeformerGeometryReadback.h
@@ -0,0 +1,15 @@
+﻿// Copyright Epic Games, Inc. All Rights Reserved.
+
+#pragma once
+
+#include "CoreMinimal.h"
+
+struct FMeshDescription;
+
+struct ENGINE_API FMeshDeformerGeometryReadbackRequest
+{
+	~FMeshDeformerGeometryReadbackRequest();
+
+	TFunction<void(const FMeshDescription&)> MeshDescriptionCallback_AnyThread;
+	bool bHandled = false;
+};
\ No newline at end of file
