[GeometryCache] Fixed incorrect material when cache contains mulitple tracks

+Add API to use a custom set of materials when creating track from mesh asset

diff --git a/Engine/Plugins/Experimental/GeometryCacheLevelSequenceBaker/Source/GeometryCacheLevelSequenceBaker/Private/GeometryCacheLevelSequenceBaker.cpp b/Engine/Plugins/Experimental/GeometryCacheLevelSequenceBaker/Source/GeometryCacheLevelSequenceBaker/Private/GeometryCacheLevelSequenceBaker.cpp
--- a/Engine/Plugins/Experimental/GeometryCacheLevelSequenceBaker/Source/GeometryCacheLevelSequenceBaker/Private/GeometryCacheLevelSequenceBaker.cpp
+++ b/Engine/Plugins/Experimental/GeometryCacheLevelSequenceBaker/Source/GeometryCacheLevelSequenceBaker/Private/GeometryCacheLevelSequenceBaker.cpp
@@ -418,9 +418,7 @@ void FGeometryCacheLevelSequenceBaker::Tick(float DeltaTime)
 			WriteToAsset();
 		}
 
-		// Make sure to wait for ongoing readback callbacks to finish
-		FScopeLock BakeTaskLifeTimeLock(&CurrentBakeTaskLifeTimeCriticalSection);
-		CurrentBakeTask = nullptr;
+		EndTask();
 	}
 }
 
@@ -544,17 +542,19 @@ void FGeometryCacheLevelSequenceBaker::Gather()
 							FSlateNotificationManager::Get().AddNotification(ErrorToast);
 
 							// clear the task and abort
-							CurrentBakeTask = nullptr;
+							EndTask();
 							return;
 						}
 					}
-					else
+
 					{
-						ComponentInfosRef.Add(Name, {Name, SkeletalMesh});
+						FComponentInfo Info;
+						Info.Name = Name;
+						Info.SkeletalMeshAsset = SkeletalMesh;
+						Info.Materials = SkeletalMeshComponent->GetMaterials();
+						ComponentInfosRef.Add(Name, MoveTemp(Info));
 					}
 
-					ComponentInfosRef.Add(Name, {Name, SkeletalMesh});
-
 					// Try to test-run deformers on these skeletal mesh component to make sure deformer shaders are compiled and ready
 					TUniquePtr<FSkeletalMeshComponentSettingScope>& ComponentSettingScope = CurrentBakeTask->GatherStageComponentSettingScopes.
 						FindOrAdd(SkeletalMeshComponent);
@@ -768,7 +768,7 @@ void FGeometryCacheLevelSequenceBaker::WriteToAsset()
 			Task.GeometrySamples[0] = Task.GeometrySamples[Task.VisibilitySamples[FirstVisibleFrameIndex].FrameIndex];
 		}
 
-		int32 TrackIndex = AddTrackWriterFromSkinnedAsset(Writer, *(Task.ComponentInfo.SkeletalMeshAsset));
+		int32 TrackIndex = AddTrackWriterFromSkinnedAssetAndMaterials(Writer, *(Task.ComponentInfo.SkeletalMeshAsset), Task.ComponentInfo.Materials);
 
 		FTrackWriter& TrackWriter = Writer.GetTrackWriter(TrackIndex);
 
@@ -946,5 +946,11 @@ void FGeometryCacheLevelSequenceBaker::FBakeTask::TickProgress()
 	SlowTask->TickProgress();
 }
 
+void FGeometryCacheLevelSequenceBaker::EndTask()
+{
+	// Make sure to wait for ongoing readback callbacks to finish
+	FScopeLock BakeTaskLifeTimeLock(&CurrentBakeTaskLifeTimeCriticalSection);
+	CurrentBakeTask = nullptr;
+}
 
 #undef LOCTEXT_NAMESPACE
diff --git a/Engine/Plugins/Experimental/GeometryCacheLevelSequenceBaker/Source/GeometryCacheLevelSequenceBaker/Private/GeometryCacheLevelSequenceBaker.h b/Engine/Plugins/Experimental/GeometryCacheLevelSequenceBaker/Source/GeometryCacheLevelSequenceBaker/Private/GeometryCacheLevelSequenceBaker.h
--- a/Engine/Plugins/Experimental/GeometryCacheLevelSequenceBaker/Source/GeometryCacheLevelSequenceBaker/Private/GeometryCacheLevelSequenceBaker.h
+++ b/Engine/Plugins/Experimental/GeometryCacheLevelSequenceBaker/Source/GeometryCacheLevelSequenceBaker/Private/GeometryCacheLevelSequenceBaker.h
@@ -13,6 +13,7 @@
 #include "GeometryCacheLevelSequenceBaker.generated.h"
 
 class UGeometryCache;
+class UMaterialInterface;
 
 UCLASS()
 class ULevelSequenceGeometryCacheBakerOption : public UObject
@@ -76,6 +77,7 @@ private:
 	{
 		FName Name;
 		USkeletalMesh* SkeletalMeshAsset;
+		TArray<TObjectPtr<UMaterialInterface>> Materials;
 	};
 	
 	struct FComponentTask
@@ -169,6 +171,8 @@ private:
 		void TickProgress();
 	};
 
+	void EndTask();
+
 	FCriticalSection CurrentBakeTaskLifeTimeCriticalSection;
 	TUniquePtr<FBakeTask> CurrentBakeTask;
 };
diff --git a/Engine/Plugins/Runtime/GeometryCache/Source/GeometryCache/Private/GeometryCacheConstantTopologyWriter.cpp b/Engine/Plugins/Runtime/GeometryCache/Source/GeometryCache/Private/GeometryCacheConstantTopologyWriter.cpp
--- a/Engine/Plugins/Runtime/GeometryCache/Source/GeometryCache/Private/GeometryCacheConstantTopologyWriter.cpp
+++ b/Engine/Plugins/Runtime/GeometryCache/Source/GeometryCache/Private/GeometryCacheConstantTopologyWriter.cpp
@@ -447,6 +447,11 @@ namespace UE::GeometryCacheHelpers
 	}
 
 	int32 AddTrackWriterFromSkinnedAsset(FGeometryCacheConstantTopologyWriter& Writer, const USkinnedAsset& Asset)
+	{
+		return AddTrackWriterFromSkinnedAssetAndMaterials(Writer, Asset, {});
+	}
+
+	int32 AddTrackWriterFromSkinnedAssetAndMaterials(FGeometryCacheConstantTopologyWriter& Writer, const USkinnedAsset& Asset, const TArray<TObjectPtr<UMaterialInterface>>& Materials)
 	{
 		const FSkeletalMeshRenderData* RenderData = Asset.GetResourceForRendering();
 		constexpr int32 LODIndex = 0;
@@ -471,7 +476,16 @@ namespace UE::GeometryCacheHelpers
 		const int32 MaterialOffset = Writer.GetNumMaterials();
 		TrackWriter.BatchesInfo = GetBatchesInfo(LODData, MaterialOffset);
 
-		Writer.AddMaterials(GetMaterialInterfaces(Asset));
+		TArray<TObjectPtr<UMaterialInterface>> AssetMaterials = GetMaterialInterfaces(Asset);
+		if (Materials.Num() == AssetMaterials.Num())
+		{
+			Writer.AddMaterials(Materials);
+		}
+		else
+		{
+			Writer.AddMaterials(AssetMaterials);
+		}
+
 		return Writer.GetNumTracks() - 1;
 	}
 
diff --git a/Engine/Plugins/Runtime/GeometryCache/Source/GeometryCache/Private/GeometryCacheSceneProxy.cpp b/Engine/Plugins/Runtime/GeometryCache/Source/GeometryCache/Private/GeometryCacheSceneProxy.cpp
--- a/Engine/Plugins/Runtime/GeometryCache/Source/GeometryCache/Private/GeometryCacheSceneProxy.cpp
+++ b/Engine/Plugins/Runtime/GeometryCache/Source/GeometryCache/Private/GeometryCacheSceneProxy.cpp
@@ -112,6 +112,19 @@ FGeometryCacheSceneProxy::FGeometryCacheSceneProxy(UGeometryCacheComponent* Comp
 
 	bCanSkipRedundantTransformUpdates = false;
 
+	// All tracks use the same array of materials since the MaterialIndex in FGeometryCacheMeshBatchInfo 
+	// is an index into the global material array UGeometryCache::Materials
+	TArray<UMaterialInterface*> Materials;
+	for (UMaterialInterface* Material : Component->GetMaterials())
+	{
+		if (Material == nullptr || !Material->CheckMaterialUsage_Concurrent(EMaterialUsage::MATUSAGE_GeometryCache))
+		{
+			Material = UMaterial::GetDefaultMaterial(MD_Surface);
+		}
+
+		Materials.Add(Material);
+	}	
+	
 	// Copy each section
 	const int32 NumTracks = Component->TrackSections.Num();
 	Tracks.Reserve(NumTracks);
@@ -151,28 +164,7 @@ FGeometryCacheSceneProxy::FGeometryCacheSceneProxy(UGeometryCacheComponent* Comp
 		NewSection->NextFrameMeshData = new FGeometryCacheMeshData();
 		NewSection->bNextFrameMeshDataSelected = false;
 
-		// Some basic sanity checks
-		for (FGeometryCacheMeshBatchInfo& BatchInfo : NewSection->MeshData->BatchesInfo)
-		{
-			UMaterialInterface* Material = Component->GetMaterial(BatchInfo.MaterialIndex);
-			if (Material == nullptr || !Material->CheckMaterialUsage_Concurrent(EMaterialUsage::MATUSAGE_GeometryCache))
-			{
-				Material = UMaterial::GetDefaultMaterial(MD_Surface);
-			}
-
-			NewSection->Materials.Add(Material);
-		}
-
-		if (NumTracks == 1)
-		{
-			// When there's only one track, it means there's one mesh (that might have been merged from other meshes and made up of multiple sections)
-			if (NewSection->Materials.Num() != Component->GetMaterials().Num())
-			{
-				// This means that the first frame does not contain all the materials used during the animation
-				// (eg. non-constant topology with increasing number of sections)
-				NewSection->Materials = Component->GetMaterials();
-			}
-		}
+		NewSection->Materials = Materials;
 
 		// Save ref to new section
 		Tracks.Add(NewSection);
diff --git a/Engine/Plugins/Runtime/GeometryCache/Source/GeometryCache/Public/GeometryCacheConstantTopologyWriter.h b/Engine/Plugins/Runtime/GeometryCache/Source/GeometryCache/Public/GeometryCacheConstantTopologyWriter.h
--- a/Engine/Plugins/Runtime/GeometryCache/Source/GeometryCache/Public/GeometryCacheConstantTopologyWriter.h
+++ b/Engine/Plugins/Runtime/GeometryCache/Source/GeometryCache/Public/GeometryCacheConstantTopologyWriter.h
@@ -135,6 +135,12 @@ namespace UE::GeometryCacheHelpers
 	 */
 	GEOMETRYCACHE_API int32 AddTrackWriterFromSkinnedAsset(FGeometryCacheConstantTopologyWriter& Writer, const USkinnedAsset& Asset);
 
+	/**
+	 * @brief This will create a track writer and fill in the track writer's data (indices, UVs, .etc) from the skinned asset, with the option to specify a set of materials 
+	 * which typically come from a specific skeletal mesh component using this skinned asset.
+	 */
+	GEOMETRYCACHE_API int32 AddTrackWriterFromSkinnedAssetAndMaterials(FGeometryCacheConstantTopologyWriter& Writer, const USkinnedAsset& Asset, const TArray<TObjectPtr<UMaterialInterface>>& Materials);
+
 	/**
 	 * @brief This will create multiple track writers and fill in the track writer's data (indices, UVs, materials .etc) from the template geometry cache. 
 	 * The number of track writers created equals to the number of tracks in the template geometry cache. 
