[GeometryCacheLevelSequenceBaker] Minor fixes:
+cache now references material assets on disk instead of dynamic instances on components causing GC issues when closing editor/ changing level
+baker now forces sequencer clock source to be tick mode to make sure it samples each frame each tick if 1 sample per frame.

diff --git a/Engine/Plugins/Animation/DeformerGraph/Source/OptimusCore/Private/OptimusDeformerDynamicInstanceManager.cpp b/Engine/Plugins/Animation/DeformerGraph/Source/OptimusCore/Private/OptimusDeformerDynamicInstanceManager.cpp
--- a/Engine/Plugins/Animation/DeformerGraph/Source/OptimusCore/Private/OptimusDeformerDynamicInstanceManager.cpp
+++ b/Engine/Plugins/Animation/DeformerGraph/Source/OptimusCore/Private/OptimusDeformerDynamicInstanceManager.cpp
@@ -28,6 +28,11 @@ void UOptimusDeformerDynamicInstanceManager::ReleaseResources()
 {
 	// Typically called during recreate render state
 	
+#if WITH_EDITORONLY_DATA
+	// Immediately release readback requests that won't be fulfilled, 
+	GeometryReadbackRequests.Reset();
+#endif // WITH_EDITORONLY_DATA	
+	
 	DefaultInstance->ReleaseResources();
 	
 	for (TPair<FGuid, TObjectPtr<UOptimusDeformerInstance>>& GuidToInstance : GuidToInstanceMap)
diff --git a/Engine/Plugins/Experimental/GeometryCacheLevelSequenceBaker/Source/GeometryCacheLevelSequenceBaker/Private/GeometryCacheLevelSequenceBaker.cpp b/Engine/Plugins/Experimental/GeometryCacheLevelSequenceBaker/Source/GeometryCacheLevelSequenceBaker/Private/GeometryCacheLevelSequenceBaker.cpp
--- a/Engine/Plugins/Experimental/GeometryCacheLevelSequenceBaker/Source/GeometryCacheLevelSequenceBaker/Private/GeometryCacheLevelSequenceBaker.cpp
+++ b/Engine/Plugins/Experimental/GeometryCacheLevelSequenceBaker/Source/GeometryCacheLevelSequenceBaker/Private/GeometryCacheLevelSequenceBaker.cpp
@@ -33,6 +33,7 @@
 #include "StaticMeshAttributes.h"
 #include "Animation/MeshDeformerGeometryReadback.h"
 #include "Async/ParallelFor.h"
+#include "Materials/MaterialInstanceDynamic.h"
 
 
 #define LOCTEXT_NAMESPACE "GeometryCacheLevelSequenceBaker"
@@ -551,7 +552,18 @@ void FGeometryCacheLevelSequenceBaker::Gather()
 						FComponentInfo Info;
 						Info.Name = Name;
 						Info.SkeletalMeshAsset = SkeletalMesh;
-						Info.Materials = SkeletalMeshComponent->GetMaterials();
+						// We need to reference the material assets instead of the dynamic instances on the component to avoid GC failure when changing level
+						for (UMaterialInterface* MaterialInterface : SkeletalMeshComponent->GetMaterials())
+						{
+							UMaterialInterface* MaterialAsset = MaterialInterface;
+							if (UMaterialInstanceDynamic* AsDynamic = Cast<UMaterialInstanceDynamic>(MaterialInterface))
+							{
+								MaterialAsset = AsDynamic->Parent;
+							}
+
+							check(MaterialAsset->IsAsset());
+							Info.Materials.Add(MaterialAsset);
+						}
 						ComponentInfosRef.Add(Name, MoveTemp(Info));
 					}
 
@@ -868,20 +880,31 @@ FGeometryCacheLevelSequenceBaker::FSequencerSettingScope::FSequencerSettingScope
 	Sequencer = InSequencer;
 	PreviousGlobalTime = InSequencer->GetGlobalTime();
 	PreviousLoopMode = InSequencer->GetSequencerSettings()->GetLoopMode();
-
 	InSequencer->GetSequencerSettings()->SetLoopMode(ESequencerLoopMode::SLM_NoLoop);
 
 	OnGlobalTimeChangedDelegateHandle = InSequencer->OnGlobalTimeChanged().AddRaw(&Get(), &FGeometryCacheLevelSequenceBaker::OnSequencerEvaluated);
+
+	UMovieSceneSequence* MovieSceneSequence = InSequencer->GetFocusedMovieSceneSequence();
+	UMovieScene* MovieScene = MovieSceneSequence->GetMovieScene();
+
+	// ensure 1 sample per tick
+	PreviousClockSource = MovieScene->GetClockSource();
+	MovieScene->SetClockSource(EUpdateClockSource::Tick);
+	InSequencer->ResetTimeController();
 }
 
 FGeometryCacheLevelSequenceBaker::FSequencerSettingScope::~FSequencerSettingScope()
 {
 	if (ensure(Sequencer.Pin()))
 	{
+		UMovieSceneSequence* MovieSceneSequence = Sequencer.Pin()->GetFocusedMovieSceneSequence();
+		UMovieScene* MovieScene = MovieSceneSequence->GetMovieScene();
+		MovieScene->SetClockSource(PreviousClockSource);
+		Sequencer.Pin()->ResetTimeController();
+		
 		Sequencer.Pin()->OnGlobalTimeChanged().Remove(OnGlobalTimeChangedDelegateHandle);
 		Sequencer.Pin()->GetSequencerSettings()->SetLoopMode(PreviousLoopMode);
 		Sequencer.Pin()->SetGlobalTime(PreviousGlobalTime.Time);
-		
 	}
 }
 
diff --git a/Engine/Plugins/Experimental/GeometryCacheLevelSequenceBaker/Source/GeometryCacheLevelSequenceBaker/Private/GeometryCacheLevelSequenceBaker.h b/Engine/Plugins/Experimental/GeometryCacheLevelSequenceBaker/Source/GeometryCacheLevelSequenceBaker/Private/GeometryCacheLevelSequenceBaker.h
--- a/Engine/Plugins/Experimental/GeometryCacheLevelSequenceBaker/Source/GeometryCacheLevelSequenceBaker/Private/GeometryCacheLevelSequenceBaker.h
+++ b/Engine/Plugins/Experimental/GeometryCacheLevelSequenceBaker/Source/GeometryCacheLevelSequenceBaker/Private/GeometryCacheLevelSequenceBaker.h
@@ -109,7 +109,7 @@ private:
 		FQualifiedFrameTime PreviousGlobalTime;
 		ESequencerLoopMode PreviousLoopMode;
 		FDelegateHandle OnGlobalTimeChangedDelegateHandle;
-
+		EUpdateClockSource PreviousClockSource;
 		
 		FSequencerSettingScope(TSharedPtr<ISequencer> InSequencer);
 		~FSequencerSettingScope();
