[GeoCacheLevelSeqBaker] Add null checks for material

diff --git a/Engine/Plugins/Experimental/GeometryCacheLevelSequenceBaker/Source/GeometryCacheLevelSequenceBaker/Private/GeometryCacheLevelSequenceBaker.cpp b/Engine/Plugins/Experimental/GeometryCacheLevelSequenceBaker/Source/GeometryCacheLevelSequenceBaker/Private/GeometryCacheLevelSequenceBaker.cpp
--- a/Engine/Plugins/Experimental/GeometryCacheLevelSequenceBaker/Source/GeometryCacheLevelSequenceBaker/Private/GeometryCacheLevelSequenceBaker.cpp
+++ b/Engine/Plugins/Experimental/GeometryCacheLevelSequenceBaker/Source/GeometryCacheLevelSequenceBaker/Private/GeometryCacheLevelSequenceBaker.cpp
@@ -557,7 +557,8 @@ void FGeometryCacheLevelSequenceBaker::Gather()
 								MaterialAsset = AsDynamic->Parent;
 							}
 
-							check(MaterialAsset->IsAsset());
+							check(MaterialAsset == nullptr || MaterialAsset->IsAsset());
+							
 							Info.Materials.Add(MaterialAsset);
 						}
 						ComponentInfosRef.Add(Name, MoveTemp(Info));
