[GeometryCacheLevelSequenceBaker] Now uses local time instead of global time such that baking from the active subsequence works

diff --git a/Engine/Plugins/Experimental/GeometryCacheLevelSequenceBaker/Source/GeometryCacheLevelSequenceBaker/Private/GeometryCacheLevelSequenceBaker.cpp b/Engine/Plugins/Experimental/GeometryCacheLevelSequenceBaker/Source/GeometryCacheLevelSequenceBaker/Private/GeometryCacheLevelSequenceBaker.cpp
--- a/Engine/Plugins/Experimental/GeometryCacheLevelSequenceBaker/Source/GeometryCacheLevelSequenceBaker/Private/GeometryCacheLevelSequenceBaker.cpp
+++ b/Engine/Plugins/Experimental/GeometryCacheLevelSequenceBaker/Source/GeometryCacheLevelSequenceBaker/Private/GeometryCacheLevelSequenceBaker.cpp
@@ -879,7 +879,7 @@ FGeometryCacheLevelSequenceBaker::FEngineFixedDeltaTimeScope::~FEngineFixedDelta
 FGeometryCacheLevelSequenceBaker::FSequencerSettingScope::FSequencerSettingScope(TSharedPtr<ISequencer> InSequencer)
 {
 	Sequencer = InSequencer;
-	PreviousGlobalTime = InSequencer->GetGlobalTime();
+	PreviousLocalTime = InSequencer->GetLocalTime();
 	PreviousLoopMode = InSequencer->GetSequencerSettings()->GetLoopMode();
 	InSequencer->GetSequencerSettings()->SetLoopMode(ESequencerLoopMode::SLM_NoLoop);
 
@@ -902,7 +902,7 @@ FGeometryCacheLevelSequenceBaker::FSequencerSettingScope::~FSequencerSettingScop
 		Sequencer.Pin()->ResetTimeController();
 		
 		Sequencer.Pin()->GetSequencerSettings()->SetLoopMode(PreviousLoopMode);
-		Sequencer.Pin()->SetGlobalTime(PreviousGlobalTime.Time);
+		Sequencer.Pin()->SetLocalTime(PreviousLocalTime.Time);
 	}
 }
 
@@ -953,15 +953,10 @@ bool FGeometryCacheLevelSequenceBaker::FBakeTask::IsSequencerPlaying()
 
 void FGeometryCacheLevelSequenceBaker::FBakeTask::PlaySequencer()
 {
-	Sequencer.Pin()->SetGlobalTime(StartFrame);
+	Sequencer.Pin()->SetLocalTime(StartFrame);
 	Sequencer.Pin()->PlayTo({EndFrame, EUpdatePositionMethod::Play});
 }
 
-void FGeometryCacheLevelSequenceBaker::FBakeTask::SequencerSnapToStart()
-{
-	Sequencer.Pin()->SetGlobalTime(StartFrame);
-}
-
 void FGeometryCacheLevelSequenceBaker::FBakeTask::UpdateGatherProgress()
 {
 	SlowTask->EnterProgressFrame(AmountOfWorkGatherStage/NumSamples);
diff --git a/Engine/Plugins/Experimental/GeometryCacheLevelSequenceBaker/Source/GeometryCacheLevelSequenceBaker/Private/GeometryCacheLevelSequenceBaker.h b/Engine/Plugins/Experimental/GeometryCacheLevelSequenceBaker/Source/GeometryCacheLevelSequenceBaker/Private/GeometryCacheLevelSequenceBaker.h
--- a/Engine/Plugins/Experimental/GeometryCacheLevelSequenceBaker/Source/GeometryCacheLevelSequenceBaker/Private/GeometryCacheLevelSequenceBaker.h
+++ b/Engine/Plugins/Experimental/GeometryCacheLevelSequenceBaker/Source/GeometryCacheLevelSequenceBaker/Private/GeometryCacheLevelSequenceBaker.h
@@ -107,7 +107,7 @@ private:
 	struct FSequencerSettingScope
 	{
 		TWeakPtr<ISequencer> Sequencer;
-		FQualifiedFrameTime PreviousGlobalTime;
+		FQualifiedFrameTime PreviousLocalTime;
 		ESequencerLoopMode PreviousLoopMode;
 		EUpdateClockSource PreviousClockSource;
 		
@@ -172,7 +172,6 @@ private:
 
 		bool IsSequencerPlaying();
 		void PlaySequencer();
-		void SequencerSnapToStart();
 		
 		void UpdateGatherProgress();
 		void UpdateBakeProgress();
